MotionDesign

diff --git a/Engine/Plugins/Experimental/ClonerEffector/Source/ClonerEffector/Private/CEClonerEffectorShared.cpp b/Engine/Plugins/Experimental/ClonerEffector/Source/ClonerEffector/Private/CEClonerEffectorShared.cpp
new file mode 100644
--- /dev/null
+++ b/Engine/Plugins/Experimental/ClonerEffector/Source/ClonerEffector/Private/CEClonerEffectorShared.cpp
@@ -0,0 +1,204 @@
+// Copyright Epic Games, Inc. All Rights Reserved.
+
+#include "CEClonerEffectorShared.h"
+
+#include "Components/BillboardComponent.h"
+#include "Engine/Texture2D.h"
+#include "NiagaraDataChannelAccessor.h"
+#include "NiagaraSystem.h"
+#include "PrimitiveSceneProxy.h"
+#include "SceneManagement.h"
+#include "UObject/Package.h"
+#include "UObject/UObjectGlobals.h"
+
+void UE::ClonerEffector::SetBillboardComponentSprite(const AActor* InActor, const FString& InTexturePath)
+{
+	if (!InActor)
+	{
+		return;
+	}
+
+	UBillboardComponent* BillboardComponent = InActor->FindComponentByClass<UBillboardComponent>();
+
+	if (!BillboardComponent)
+	{
+		return;
+	}
+
+	UTexture2D* SpriteTexture = FindObject<UTexture2D>(nullptr, *InTexturePath);
+
+	if (!SpriteTexture)
+	{
+		SpriteTexture = LoadObject<UTexture2D>(nullptr, *InTexturePath);
+	}
+
+	if (SpriteTexture && SpriteTexture != BillboardComponent->Sprite)
+	{
+		BillboardComponent->SetSprite(SpriteTexture);
+	}
+}
+
+void UE::ClonerEffector::SetBillboardComponentVisibility(const AActor* InActor, bool bInVisibility)
+{
+	if (!InActor)
+	{
+		return;
+	}
+
+	UBillboardComponent* BillboardComponent = InActor->FindComponentByClass<UBillboardComponent>();
+
+	if (!BillboardComponent)
+	{
+		return;
+	}
+
+	BillboardComponent->SetVisibility(bInVisibility, false);
+}
+
+void FCEClonerEffectorChannelData::Write(UNiagaraDataChannelWriter* InWriter) const
+{
+	if (!InWriter)
+	{
+		return;
+	}
+
+	/** General */
+	InWriter->WriteInt(EasingName, Identifier, static_cast<int32>(Easing));
+	InWriter->WriteInt(ModeName, Identifier, static_cast<int32>(Mode));
+	InWriter->WriteInt(TypeName, Identifier, static_cast<int32>(Type));
+	InWriter->WriteFloat(MagnitudeName, Identifier, Magnitude);
+	InWriter->WriteVector(InnerExtentName, Identifier, InnerExtent);
+	InWriter->WriteVector(OuterExtentName, Identifier, OuterExtent);
+	InWriter->WriteVector(LocationDeltaName, Identifier, LocationDelta);
+	InWriter->WriteQuat(RotationDeltaName, Identifier, RotationDelta);
+	InWriter->WriteVector(ScaleDeltaName, Identifier, ScaleDelta);
+	InWriter->WritePosition(LocationName, Identifier, Location);
+	InWriter->WriteQuat(RotationName, Identifier, Rotation);
+	InWriter->WriteVector(ScaleName, Identifier, Scale);
+	InWriter->WriteFloat(FrequencyName, Identifier, Frequency);
+	InWriter->WriteVector(PanName, Identifier, Pan);
+
+	/** Forces */
+	InWriter->WriteFloat(OrientationForceRateName, Identifier, OrientationForceRate);
+	InWriter->WriteVector(OrientationForceMinName, Identifier, OrientationForceMin);
+	InWriter->WriteVector(OrientationForceMaxName, Identifier, OrientationForceMax);
+	InWriter->WriteFloat(VortexForceAmountName, Identifier, VortexForceAmount);
+	InWriter->WriteVector(VortexForceAxisName, Identifier, VortexForceAxis);
+	InWriter->WriteFloat(CurlNoiseForceStrengthName, Identifier, CurlNoiseForceStrength);
+	InWriter->WriteFloat(CurlNoiseForceFrequencyName, Identifier, CurlNoiseForceFrequency);
+	InWriter->WriteFloat(AttractionForceStrengthName, Identifier, AttractionForceStrength);
+	InWriter->WriteFloat(AttractionForceFalloffName, Identifier, AttractionForceFalloff);
+	InWriter->WriteVector(GravityForceAccelerationName, Identifier, GravityForceAcceleration);
+}
+
+void FCEClonerEffectorChannelData::Read(const UNiagaraDataChannelReader* InReader)
+{
+	if (!InReader)
+	{
+		return;
+	}
+
+	/** Effector */
+	bool bIsValid;
+	Easing = static_cast<ECEClonerEasing>(InReader->ReadInt(EasingName, Identifier, bIsValid));
+	Mode = static_cast<ECEClonerEffectorMode>(InReader->ReadInt(ModeName, Identifier, bIsValid));
+	Type = static_cast<ECEClonerEffectorType>(InReader->ReadInt(TypeName, Identifier, bIsValid));
+	Magnitude = InReader->ReadFloat(MagnitudeName, Identifier, bIsValid);
+	InnerExtent = InReader->ReadVector(InnerExtentName, Identifier, bIsValid);
+	OuterExtent = InReader->ReadVector(OuterExtentName, Identifier, bIsValid);
+	LocationDelta = InReader->ReadVector(LocationDeltaName, Identifier, bIsValid);
+	RotationDelta = InReader->ReadQuat(RotationDeltaName, Identifier, bIsValid);
+	ScaleDelta = InReader->ReadVector(ScaleDeltaName, Identifier, bIsValid);
+	Location = InReader->ReadPosition(LocationName, Identifier, bIsValid);
+	Rotation = InReader->ReadQuat(RotationName, Identifier, bIsValid);
+	Scale = InReader->ReadVector(ScaleName, Identifier, bIsValid);
+	Frequency = InReader->ReadFloat(FrequencyName, Identifier, bIsValid);
+	Pan = InReader->ReadVector(PanName, Identifier, bIsValid);
+
+	/** Forces */
+	OrientationForceRate = InReader->ReadFloat(OrientationForceRateName, Identifier, bIsValid);
+	OrientationForceMin = InReader->ReadVector(OrientationForceMinName, Identifier, bIsValid);
+	OrientationForceMax = InReader->ReadVector(OrientationForceMaxName, Identifier, bIsValid);
+	VortexForceAmount = InReader->ReadFloat(VortexForceAmountName, Identifier, bIsValid);
+	VortexForceAxis = InReader->ReadVector(VortexForceAxisName, Identifier, bIsValid);
+	CurlNoiseForceStrength = InReader->ReadFloat(CurlNoiseForceStrengthName, Identifier, bIsValid);
+	CurlNoiseForceFrequency = InReader->ReadFloat(CurlNoiseForceFrequencyName, Identifier, bIsValid);
+	AttractionForceStrength = InReader->ReadFloat(AttractionForceStrengthName, Identifier, bIsValid);
+	AttractionForceFalloff = InReader->ReadFloat(AttractionForceFalloffName, Identifier, bIsValid);
+	GravityForceAcceleration = InReader->ReadVector(GravityForceAccelerationName, Identifier, bIsValid);
+}
+
+FCEClonerEffectorDataInterfaces::FCEClonerEffectorDataInterfaces(const UNiagaraSystem* InSystem)
+{
+	check(InSystem);
+	const FNiagaraUserRedirectionParameterStore& UserParameterStore = InSystem->GetExposedParameters();
+
+	static const FNiagaraVariable EffectorIndexDIVar(FNiagaraTypeDefinition(UNiagaraDataInterfaceArrayInt32::StaticClass()), FCEClonerEffectorDataInterfaces::IndexName);
+	UNiagaraDataInterfaceArrayInt32* IndexArrayDI = CastChecked<UNiagaraDataInterfaceArrayInt32>(UserParameterStore.GetDataInterface(EffectorIndexDIVar));
+	DataInterfaces.Add(EffectorIndexDIVar.GetName(), IndexArrayDI);
+}
+
+void FCEClonerEffectorDataInterfaces::CopyTo(FCEClonerEffectorDataInterfaces& InOther) const
+{
+	for (const TPair<FName, TObjectPtr<UNiagaraDataInterface>>& DataInterfacePair : DataInterfaces)
+	{
+		const TObjectPtr<UNiagaraDataInterface> DataInterface = DataInterfacePair.Value;
+		const TObjectPtr<UNiagaraDataInterface>* OtherDataInterface = InOther.DataInterfaces.Find(DataInterfacePair.Key);
+
+		if (!OtherDataInterface || !DataInterface)
+		{
+			continue;
+		}
+
+		DataInterface->CopyTo(*OtherDataInterface);
+	}
+}
+
+void FCEClonerEffectorDataInterfaces::Resize(int32 InSize) const
+{
+	if (UNiagaraDataInterfaceArrayInt32* IndexArray = GetIndexArray())
+	{
+		IndexArray->GetArrayReference().SetNum(InSize);
+	}
+}
+
+void FCEClonerEffectorDataInterfaces::Remove(int32 InIndex) const
+{
+	UNiagaraDataInterfaceArrayInt32* IndexArray = GetIndexArray();
+	if (IndexArray && IndexArray->GetArrayReference().IsValidIndex(InIndex))
+	{
+		IndexArray->GetArrayReference().RemoveAt(InIndex);
+	}
+}
+
+bool FCEClonerEffectorDataInterfaces::IsValid() const
+{
+	for (const TPair<FName, TObjectPtr<UNiagaraDataInterface>>& DataInterfacePair : DataInterfaces)
+	{
+		if (!DataInterfacePair.Value.Get())
+		{
+			return false;
+		}
+	}
+
+	return !DataInterfaces.IsEmpty();
+}
+
+int32 FCEClonerEffectorDataInterfaces::Num() const
+{
+	UNiagaraDataInterfaceArrayInt32* IndexDI = GetIndexArray();
+
+	const int32 IndexSize = IndexDI->GetArrayReference().Num();
+
+	return IndexSize;
+}
+
+UNiagaraDataInterfaceArrayInt32* FCEClonerEffectorDataInterfaces::GetIndexArray() const
+{
+	if (const TObjectPtr<UNiagaraDataInterface>* Array = DataInterfaces.Find(FCEClonerEffectorDataInterfaces::IndexName))
+	{
+		return Cast<UNiagaraDataInterfaceArrayInt32>(Array->Get());
+	}
+
+	return nullptr;
+}
diff --git a/Engine/Plugins/Experimental/ClonerEffector/Source/ClonerEffector/Private/CEModule.cpp b/Engine/Plugins/Experimental/ClonerEffector/Source/ClonerEffector/Private/CEModule.cpp
new file mode 100644
--- /dev/null
+++ b/Engine/Plugins/Experimental/ClonerEffector/Source/ClonerEffector/Private/CEModule.cpp
@@ -0,0 +1,5 @@
+﻿// Copyright Epic Games, Inc. All Rights Reserved.
+
+#include "Modules/ModuleManager.h"
+
+IMPLEMENT_MODULE(FDefaultModuleImpl, ClonerEffector)
\ No newline at end of file
diff --git a/Engine/Plugins/Experimental/ClonerEffector/Source/ClonerEffector/Private/Cloner/CEClonerActor.cpp b/Engine/Plugins/Experimental/ClonerEffector/Source/ClonerEffector/Private/Cloner/CEClonerActor.cpp
new file mode 100644
--- /dev/null
+++ b/Engine/Plugins/Experimental/ClonerEffector/Source/ClonerEffector/Private/Cloner/CEClonerActor.cpp
@@ -0,0 +1,1404 @@
+// Copyright Epic Games, Inc. All Rights Reserved.
+
+#include "Cloner/CEClonerActor.h"
+
+#include "Cloner/CEClonerComponent.h"
+#include "Cloner/Layouts/CEClonerCircleLayout.h"
+#include "Cloner/Layouts/CEClonerCylinderLayout.h"
+#include "Cloner/Layouts/CEClonerGridLayout.h"
+#include "Cloner/Layouts/CEClonerHoneycombLayout.h"
+#include "Cloner/Layouts/CEClonerLineLayout.h"
+#include "Cloner/Layouts/CEClonerMeshLayout.h"
+#include "Cloner/Layouts/CEClonerSphereUniformLayout.h"
+#include "Cloner/Layouts/CEClonerSplineLayout.h"
+#include "Effector/CEEffectorActor.h"
+#include "Engine/StaticMeshActor.h"
+#include "NiagaraDataInterfaceCurve.h"
+#include "NiagaraDataInterfaceSkeletalMesh.h"
+#include "NiagaraMeshRendererProperties.h"
+#include "NiagaraSystem.h"
+#include "Subsystems/CEClonerSubsystem.h"
+#include "UObject/ConstructorHelpers.h"
+
+#if WITH_EDITOR
+#include "Editor.h"
+#include "Selection.h"
+#include "Settings/LevelEditorViewportSettings.h"
+#endif
+
+DEFINE_LOG_CATEGORY_STATIC(LogCEClonerActor, Log, All);
+
+// Sets default values
+ACEClonerActor::ACEClonerActor()
+{
+	SetCanBeDamaged(false);
+	PrimaryActorTick.bCanEverTick          = true;
+	PrimaryActorTick.bStartWithTickEnabled = true;
+	PrimaryActorTick.bTickEvenWhenPaused   = true;
+
+	ClonerComponent = CreateDefaultSubobject<UCEClonerComponent>(TEXT("AvaClonerComponent"));
+	SetRootComponent(ClonerComponent);
+
+	// Default Scale Curve
+	LifetimeScaleCurve.AddKey(0, 1.f);
+	LifetimeScaleCurve.AddKey(1, 0.f);
+
+	if (!IsTemplate())
+	{
+#if WITH_EDITOR
+		if (GEditor)
+		{
+			GEditor->GetSelectedActors()->SelectionChangedEvent.AddUObject(this, &ACEClonerActor::OnEditorSelectionChanged);
+		}
+#endif
+
+		ClonerComponent->TransformUpdated.AddUObject(this, &ACEClonerActor::OnClonerTransformed);
+		ClonerComponent->OnClonerMeshUpdated.AddUObject(this, &ACEClonerActor::OnClonerMeshUpdated);
+
+		const TArray<FString> LayoutNames = GetClonerLayoutNames();
+
+		// Apply default layout
+		LayoutName = !LayoutNames.IsEmpty() ? FName(LayoutNames[0]) : NAME_None;
+
+#if WITH_EDITOR
+		if (const UCEClonerSubsystem* ClonerSubsystem = UCEClonerSubsystem::Get())
+		{
+			UCEClonerSubsystem::OnCVarChanged().AddUObject(this, &ACEClonerActor::OnCVarChanged);
+			bReduceMotionGhosting = ClonerSubsystem->IsNoFlickerEnabled();
+		}
+#endif
+	}
+}
+
+#if WITH_EDITOR
+FString ACEClonerActor::GetDefaultActorLabel() const
+{
+	return DefaultLabel;
+}
+
+const TCEPropertyChangeDispatcher<ACEClonerActor> ACEClonerActor::PropertyChangeDispatcher =
+{
+	{ GET_MEMBER_NAME_CHECKED(ACEClonerActor, bEnabled), &ACEClonerActor::OnEnabledChanged },
+	{ GET_MEMBER_NAME_CHECKED(ACEClonerActor, Seed), &ACEClonerActor::OnSeedChanged },
+	/** Layout */
+	{ GET_MEMBER_NAME_CHECKED(ACEClonerActor, LayoutName), &ACEClonerActor::OnLayoutNameChanged },
+	/** Advanced */
+	{ GET_MEMBER_NAME_CHECKED(ACEClonerActor, bVisualizerSpriteVisible), &ACEClonerActor::OnVisualizerSpriteVisibleChanged },
+	{ GET_MEMBER_NAME_CHECKED(ACEClonerActor, bReduceMotionGhosting), &ACEClonerActor::OnReduceMotionGhostingChanged },
+	/** Spawn */
+	{ GET_MEMBER_NAME_CHECKED(ACEClonerActor, SpawnLoopMode), &ACEClonerActor::OnSpawnOptionsChanged },
+	{ GET_MEMBER_NAME_CHECKED(ACEClonerActor, SpawnLoopInterval), &ACEClonerActor::OnSpawnOptionsChanged },
+	{ GET_MEMBER_NAME_CHECKED(ACEClonerActor, SpawnLoopIterations), &ACEClonerActor::OnSpawnOptionsChanged },
+	{ GET_MEMBER_NAME_CHECKED(ACEClonerActor, SpawnBehaviorMode), &ACEClonerActor::OnSpawnOptionsChanged },
+	{ GET_MEMBER_NAME_CHECKED(ACEClonerActor, SpawnRate), &ACEClonerActor::OnSpawnOptionsChanged },
+	/** Lifetime */
+	{ GET_MEMBER_NAME_CHECKED(ACEClonerActor, bLifetimeEnabled), &ACEClonerActor::OnLifetimeOptionsChanged },
+	{ GET_MEMBER_NAME_CHECKED(ACEClonerActor, LifetimeMin), &ACEClonerActor::OnLifetimeOptionsChanged },
+	{ GET_MEMBER_NAME_CHECKED(ACEClonerActor, LifetimeMax), &ACEClonerActor::OnLifetimeOptionsChanged },
+	{ GET_MEMBER_NAME_CHECKED(ACEClonerActor, bLifetimeScaleEnabled), &ACEClonerActor::OnLifetimeOptionsChanged },
+	/** Renderer */
+	{ GET_MEMBER_NAME_CHECKED(ACEClonerActor, MeshRenderMode), &ACEClonerActor::OnMeshRenderModeChanged },
+	{ GET_MEMBER_NAME_CHECKED(ACEClonerActor, MeshFacingMode), &ACEClonerActor::OnMeshRendererOptionsChanged },
+	{ GET_MEMBER_NAME_CHECKED(ACEClonerActor, bMeshCastShadows), &ACEClonerActor::OnMeshRendererOptionsChanged },
+	{ GET_MEMBER_NAME_CHECKED(ACEClonerActor, DefaultMeshes), &ACEClonerActor::OnDefaultMeshesChanged },
+	/** Progress */
+	{ GET_MEMBER_NAME_CHECKED(ACEClonerActor, bInvertProgress), &ACEClonerActor::OnProgressChanged },
+	{ GET_MEMBER_NAME_CHECKED(ACEClonerActor, Progress), &ACEClonerActor::OnProgressChanged },
+	/** Step */
+	{ GET_MEMBER_NAME_CHECKED(ACEClonerActor, bDeltaStepEnabled), &ACEClonerActor::OnDeltaStepChanged },
+	{ GET_MEMBER_NAME_CHECKED(ACEClonerActor, DeltaStepRotation), &ACEClonerActor::OnDeltaStepChanged },
+	{ GET_MEMBER_NAME_CHECKED(ACEClonerActor, DeltaStepScale), &ACEClonerActor::OnDeltaStepChanged },
+	/** Range */
+	{ GET_MEMBER_NAME_CHECKED(ACEClonerActor, bRangeEnabled), &ACEClonerActor::OnRangeOptionsChanged },
+	{ GET_MEMBER_NAME_CHECKED(ACEClonerActor, RangeOffsetMin), &ACEClonerActor::OnRangeOptionsChanged },
+	{ GET_MEMBER_NAME_CHECKED(ACEClonerActor, RangeOffsetMax), &ACEClonerActor::OnRangeOptionsChanged },
+	{ GET_MEMBER_NAME_CHECKED(ACEClonerActor, RangeRotationMin), &ACEClonerActor::OnRangeOptionsChanged },
+	{ GET_MEMBER_NAME_CHECKED(ACEClonerActor, RangeRotationMax), &ACEClonerActor::OnRangeOptionsChanged },
+	{ GET_MEMBER_NAME_CHECKED(ACEClonerActor, RangeScaleMin), &ACEClonerActor::OnRangeOptionsChanged },
+	{ GET_MEMBER_NAME_CHECKED(ACEClonerActor, RangeScaleMax), &ACEClonerActor::OnRangeOptionsChanged },
+	{ GET_MEMBER_NAME_CHECKED(ACEClonerActor, bRangeScaleUniform), &ACEClonerActor::OnRangeOptionsChanged },
+	{ GET_MEMBER_NAME_CHECKED(ACEClonerActor, RangeScaleUniformMin), &ACEClonerActor::OnRangeOptionsChanged },
+	{ GET_MEMBER_NAME_CHECKED(ACEClonerActor, RangeScaleUniformMax), &ACEClonerActor::OnRangeOptionsChanged },
+};
+
+void ACEClonerActor::PostEditChangeProperty(FPropertyChangedEvent& InPropertyChangedEvent)
+{
+	Super::PostEditChangeProperty(InPropertyChangedEvent);
+
+	PropertyChangeDispatcher.OnPropertyChanged(this, InPropertyChangedEvent);
+}
+#endif
+
+void ACEClonerActor::PostInitProperties()
+{
+	Super::PostInitProperties();
+
+	if (HasAnyFlags(RF_ClassDefaultObject))
+	{
+		// Register new type def for niagara
+
+		constexpr ENiagaraTypeRegistryFlags MeshFlags =
+			ENiagaraTypeRegistryFlags::AllowAnyVariable |
+			ENiagaraTypeRegistryFlags::AllowParameter;
+
+		FNiagaraTypeRegistry::Register(FNiagaraTypeDefinition(StaticEnum<ECEClonerMeshRenderMode>()), MeshFlags);
+		FNiagaraTypeRegistry::Register(FNiagaraTypeDefinition(StaticEnum<ECEClonerGridConstraint>()), MeshFlags);
+		FNiagaraTypeRegistry::Register(FNiagaraTypeDefinition(StaticEnum<ECEClonerPlane>()), MeshFlags);
+		FNiagaraTypeRegistry::Register(FNiagaraTypeDefinition(StaticEnum<ECEClonerAxis>()), MeshFlags);
+		FNiagaraTypeRegistry::Register(FNiagaraTypeDefinition(StaticEnum<ECEClonerEasing>()), MeshFlags);
+		FNiagaraTypeRegistry::Register(FNiagaraTypeDefinition(StaticEnum<ECEClonerMeshAsset>()), MeshFlags);
+		FNiagaraTypeRegistry::Register(FNiagaraTypeDefinition(StaticEnum<ECEClonerMeshSampleData>()), MeshFlags);
+		FNiagaraTypeRegistry::Register(FNiagaraTypeDefinition(StaticEnum<ECEClonerEffectorType>()), MeshFlags);
+		FNiagaraTypeRegistry::Register(FNiagaraTypeDefinition(StaticEnum<ECEClonerTextureSampleChannel>()), MeshFlags);
+		FNiagaraTypeRegistry::Register(FNiagaraTypeDefinition(StaticEnum<ECEClonerCompareMode>()), MeshFlags);
+		FNiagaraTypeRegistry::Register(FNiagaraTypeDefinition(StaticEnum<ECEClonerEffectorMode>()), MeshFlags);
+		FNiagaraTypeRegistry::Register(FNiagaraTypeDefinition(StaticEnum<ECEClonerSpawnLoopMode>()), MeshFlags);
+		FNiagaraTypeRegistry::Register(FNiagaraTypeDefinition(StaticEnum<ECEClonerSpawnBehaviorMode>()), MeshFlags);
+	}
+}
+
+void ACEClonerActor::PostLoad()
+{
+	Super::PostLoad();
+
+	// Migrate old properties to new layout system
+	if (!bDeprecatedPropertiesMigrated)
+	{
+		// Grid options
+		if (UCEClonerGridLayout* GridLayout = FindOrAddLayout<UCEClonerGridLayout>())
+		{
+			GridLayout->SetCountX(GridOptions_DEPRECATED.Count.X);
+			GridLayout->SetCountY(GridOptions_DEPRECATED.Count.Y);
+			GridLayout->SetCountZ(GridOptions_DEPRECATED.Count.Z);
+			GridLayout->SetSpacingX(GridOptions_DEPRECATED.Spacing.X);
+			GridLayout->SetSpacingY(GridOptions_DEPRECATED.Spacing.Y);
+			GridLayout->SetSpacingZ(GridOptions_DEPRECATED.Spacing.Z);
+			GridLayout->SetConstraint(GridOptions_DEPRECATED.Constraint);
+			GridLayout->SetInvertConstraint(GridOptions_DEPRECATED.bInvertConstraint);
+			GridLayout->SetSphereConstraint(GridOptions_DEPRECATED.SphereConstraint);
+			GridLayout->SetCylinderConstraint(GridOptions_DEPRECATED.CylinderConstraint);
+			GridLayout->SetTextureConstraint(GridOptions_DEPRECATED.TextureConstraint);
+		}
+
+		// Line options
+		if (UCEClonerLineLayout* LineLayout = FindOrAddLayout<UCEClonerLineLayout>())
+		{
+			LineLayout->SetCount(LineOptions_DEPRECATED.Count);
+			LineLayout->SetSpacing(LineOptions_DEPRECATED.Spacing);
+			LineLayout->SetAxis(LineOptions_DEPRECATED.Axis);
+			LineLayout->SetDirection(LineOptions_DEPRECATED.Direction);
+			LineLayout->SetRotation(LineOptions_DEPRECATED.Rotation);
+		}
+
+		// Circle options
+		if (UCEClonerCircleLayout* CircleLayout = FindOrAddLayout<UCEClonerCircleLayout>())
+		{
+			CircleLayout->SetCount(CircleOptions_DEPRECATED.Count);
+			CircleLayout->SetRadius(CircleOptions_DEPRECATED.Radius);
+			CircleLayout->SetAngleStart(CircleOptions_DEPRECATED.AngleStart);
+			CircleLayout->SetAngleRatio(CircleOptions_DEPRECATED.AngleRatio);
+			CircleLayout->SetOrientMesh(CircleOptions_DEPRECATED.bOrientMesh);
+			CircleLayout->SetPlane(CircleOptions_DEPRECATED.Plane);
+			CircleLayout->SetRotation(CircleOptions_DEPRECATED.Rotation);
+			CircleLayout->SetScale(CircleOptions_DEPRECATED.Scale);
+		}
+
+		// Cylinder options
+		if (UCEClonerCylinderLayout* CylinderLayout = FindOrAddLayout<UCEClonerCylinderLayout>())
+		{
+			CylinderLayout->SetBaseCount(CylinderOptions_DEPRECATED.BaseCount);
+			CylinderLayout->SetHeightCount(CylinderOptions_DEPRECATED.HeightCount);
+			CylinderLayout->SetHeight(CylinderOptions_DEPRECATED.Height);
+			CylinderLayout->SetRadius(CylinderOptions_DEPRECATED.Radius);
+			CylinderLayout->SetAngleStart(CylinderOptions_DEPRECATED.AngleStart);
+			CylinderLayout->SetAngleRatio(CylinderOptions_DEPRECATED.AngleRatio);
+			CylinderLayout->SetOrientMesh(CylinderOptions_DEPRECATED.bOrientMesh);
+			CylinderLayout->SetPlane(CylinderOptions_DEPRECATED.Plane);
+			CylinderLayout->SetRotation(CylinderOptions_DEPRECATED.Rotation);
+			CylinderLayout->SetScale(CylinderOptions_DEPRECATED.Scale);
+		}
+
+		// Sphere options
+		if (UCEClonerSphereUniformLayout* SphereLayout = FindOrAddLayout<UCEClonerSphereUniformLayout>())
+		{
+			SphereLayout->SetCount(SphereOptions_DEPRECATED.Count);
+			SphereLayout->SetRadius(SphereOptions_DEPRECATED.Radius);
+			SphereLayout->SetRatio(SphereOptions_DEPRECATED.Ratio);
+			SphereLayout->SetOrientMesh(SphereOptions_DEPRECATED.bOrientMesh);
+			SphereLayout->SetRotation(SphereOptions_DEPRECATED.Rotation);
+			SphereLayout->SetScale(SphereOptions_DEPRECATED.Scale);
+		}
+
+		// Honeycomb options
+		if (UCEClonerHoneycombLayout* HoneycombLayout = FindOrAddLayout<UCEClonerHoneycombLayout>())
+		{
+			HoneycombLayout->SetPlane(HoneycombOptions_DEPRECATED.Plane);
+			HoneycombLayout->SetWidthCount(HoneycombOptions_DEPRECATED.WidthCount);
+			HoneycombLayout->SetHeightCount(HoneycombOptions_DEPRECATED.HeightCount);
+			HoneycombLayout->SetWidthOffset(HoneycombOptions_DEPRECATED.WidthOffset);
+			HoneycombLayout->SetHeightOffset(HoneycombOptions_DEPRECATED.HeightOffset);
+			HoneycombLayout->SetHeightSpacing(HoneycombOptions_DEPRECATED.HeightSpacing);
+			HoneycombLayout->SetWidthSpacing(HoneycombOptions_DEPRECATED.WidthSpacing);
+		}
+
+		// Sample mesh options
+		if (UCEClonerMeshLayout* MeshLayout = FindOrAddLayout<UCEClonerMeshLayout>())
+		{
+			MeshLayout->SetCount(SampleMeshOptions_DEPRECATED.Count);
+			MeshLayout->SetAsset(SampleMeshOptions_DEPRECATED.Asset);
+			MeshLayout->SetSampleData(SampleMeshOptions_DEPRECATED.SampleData);
+			MeshLayout->SetSampleActorWeak(SampleMeshOptions_DEPRECATED.SampleActor);
+		}
+
+		// Sample spline options
+		if (UCEClonerSplineLayout* SplineLayout = FindOrAddLayout<UCEClonerSplineLayout>())
+		{
+			SplineLayout->SetCount(SampleSplineOptions_DEPRECATED.Count);
+			SplineLayout->SetSplineActorWeak(SampleSplineOptions_DEPRECATED.SplineActor);
+		}
+
+		// Set previous layout
+		if (const UCEClonerSubsystem* Subsystem = UCEClonerSubsystem::Get())
+		{
+			const int32 LayoutIdx = static_cast<int32>(Layout_DEPRECATED);
+			SetLayoutName(Subsystem->GetLayoutNames()[LayoutIdx]);
+		}
+
+		bDeprecatedPropertiesMigrated = true;
+
+		UE_LOG(LogCEClonerActor, Log, TEXT("%s : Cloner deprecated properties migrated"), *GetActorNameOrLabel());
+	}
+
+	for (TMap<FName, TObjectPtr<UCEClonerLayoutBase>>::TIterator It(LayoutInstances); It; ++It)
+	{
+		if (!It->Value)
+		{
+			It.RemoveCurrent();
+			continue;
+		}
+
+		// Update layout name if changed
+		const FName NewLayoutName = It->Value->GetLayoutName();
+
+		if (It->Key != NewLayoutName)
+		{
+			LayoutInstances.Add(NewLayoutName, It->Value);
+			It.RemoveCurrent();
+
+			UE_LOG(LogCEClonerActor, Log, TEXT("%s : Cloner layout name changed %s"), *GetActorNameOrLabel(), *NewLayoutName.ToString());
+		}
+	}
+}
+
+void ACEClonerActor::PostActorCreated()
+{
+	Super::PostActorCreated();
+
+#if WITH_EDITOR
+	SpawnDefaultActorAttached();
+#endif
+}
+
+void ACEClonerActor::Tick(float InDeltaTime)
+{
+	Super::Tick(InDeltaTime);
+
+	if (!bClonerInitialized)
+	{
+		InitializeCloner();
+	}
+
+	if (!bEnabled)
+	{
+		return;
+	}
+
+	TreeUpdateDeltaTime += InDeltaTime;
+	if (TreeUpdateDeltaTime >= TreeUpdateInterval)
+	{
+		TreeUpdateDeltaTime -= TreeUpdateInterval;
+
+		if (ClonerComponent)
+		{
+			ClonerComponent->UpdateClonerAttachmentTree();
+			ClonerComponent->UpdateClonerRenderState();
+		}
+	}
+
+	if (bNeedsRefresh)
+	{
+		bNeedsRefresh = false;
+		RequestClonerUpdate(true);
+	}
+}
+
+void ACEClonerActor::UpdateLayoutOptions()
+{
+	OnMeshRenderModeChanged();
+	OnSeedChanged();
+
+	// update layouts options
+	if (UCEClonerLayoutBase* CurrentLayout = ClonerComponent->GetClonerActiveLayout())
+	{
+		CurrentLayout->UpdateLayoutParameters();
+	}
+
+	OnMeshRendererOptionsChanged();
+	OnDefaultMeshesChanged();
+	OnDeltaStepChanged();
+	OnRangeOptionsChanged();
+	OnSpawnOptionsChanged();
+	OnLifetimeOptionsChanged();
+}
+
+void ACEClonerActor::UpdateClonerEffectors()
+{
+	Effectors.RemoveAll([](const TWeakObjectPtr<ACEEffectorActor>& InEffector)
+	{
+		return !InEffector.IsValid();
+	});
+
+	if (const UCEClonerLayoutBase* LayoutSystem = GetActiveLayout())
+	{
+		// Effectors could be registered before system is loaded
+		const int32 EffectorCount = GetEffectorCount();
+
+		if (LayoutSystem->GetDataInterfaces().Num() != EffectorCount)
+		{
+			LayoutSystem->GetDataInterfaces().Resize(EffectorCount);
+
+			RequestClonerUpdate();
+		}
+	}
+
+	ForEachEffector([this](ACEEffectorActor* InEffector, int32 InIdx)
+	{
+		InEffector->OnClonerUpdated(this, InIdx);
+		return true;
+	});
+}
+
+void ACEClonerActor::SetTreeUpdateInterval(float InInterval)
+{
+	if (InInterval == TreeUpdateInterval)
+	{
+		return;
+	}
+
+	TreeUpdateInterval = InInterval;
+}
+
+void ACEClonerActor::SetMeshRenderMode(ECEClonerMeshRenderMode InMode)
+{
+	if (InMode == MeshRenderMode)
+	{
+		return;
+	}
+
+	MeshRenderMode = InMode;
+	OnMeshRenderModeChanged();
+}
+
+void ACEClonerActor::SetMeshFacingMode(ENiagaraMeshFacingMode InMode)
+{
+	if (MeshFacingMode == InMode)
+	{
+		return;
+	}
+
+	MeshFacingMode = InMode;
+	OnMeshRendererOptionsChanged();
+}
+
+void ACEClonerActor::SetMeshCastShadows(bool InbCastShadows)
+{
+	if (bMeshCastShadows == InbCastShadows)
+	{
+		return;
+	}
+
+	bMeshCastShadows = InbCastShadows;
+	OnMeshRendererOptionsChanged();
+}
+
+void ACEClonerActor::SetDefaultMeshes(const TArray<TObjectPtr<UStaticMesh>>& InMeshes)
+{
+	DefaultMeshes = InMeshes;
+	OnDefaultMeshesChanged();
+}
+
+void ACEClonerActor::BP_SetDefaultMeshes(const TArray<UStaticMesh*>& InMeshes)
+{
+	DefaultMeshes.Empty(InMeshes.Num());
+	Algo::Transform(InMeshes, DefaultMeshes, [](UStaticMesh* InMesh)->TObjectPtr<UStaticMesh>
+	{
+		return InMesh;
+	});
+	OnDefaultMeshesChanged();
+}
+
+TArray<UStaticMesh*> ACEClonerActor::BP_GetDefaultMeshes() const
+{
+	TArray<UStaticMesh*> Meshes;
+	Meshes.Reserve(DefaultMeshes.Num());
+	Algo::Transform(DefaultMeshes, Meshes, [](const TObjectPtr<UStaticMesh>& InMesh)->UStaticMesh*
+	{
+		return InMesh;
+	});
+	return Meshes;
+}
+
+void ACEClonerActor::SetSeed(int32 InSeed)
+{
+	if (InSeed == Seed)
+	{
+		return;
+	}
+
+	Seed = InSeed;
+	OnSeedChanged();
+}
+
+void ACEClonerActor::SetDeltaStepEnabled(bool bInEnabled)
+{
+	if (bDeltaStepEnabled == bInEnabled)
+	{
+		return;
+	}
+
+	bDeltaStepEnabled = bInEnabled;
+	OnDeltaStepChanged();
+}
+
+void ACEClonerActor::SetDeltaStepRotation(const FRotator& InRotation)
+{
+	if (InRotation == DeltaStepRotation)
+	{
+		return;
+	}
+
+	DeltaStepRotation = InRotation;
+	OnDeltaStepChanged();
+}
+
+void ACEClonerActor::SetDeltaStepScale(const FVector& InScale)
+{
+	if (InScale == DeltaStepScale)
+	{
+		return;
+	}
+
+	DeltaStepScale = InScale;
+	OnDeltaStepChanged();
+}
+
+void ACEClonerActor::SetEnabled(bool bInEnable)
+{
+	if (bInEnable == bEnabled)
+	{
+		return;
+	}
+
+	bEnabled = bInEnable;
+	OnEnabledChanged();
+}
+
+void ACEClonerActor::SetInvertProgress(bool bInInvertProgress)
+{
+	if (bInvertProgress == bInInvertProgress)
+	{
+		return;
+	}
+
+	bInvertProgress = bInInvertProgress;
+	OnProgressChanged();
+}
+
+void ACEClonerActor::SetProgress(float InProgress)
+{
+	if (Progress == InProgress)
+	{
+		return;
+	}
+
+	if (InProgress < 0.f || InProgress > 1.f)
+	{
+		return;
+	}
+
+	Progress = InProgress;
+	OnProgressChanged();
+}
+
+void ACEClonerActor::SetLayoutName(FName InLayoutName)
+{
+	if (LayoutName == InLayoutName)
+	{
+		return;
+	}
+
+	const TArray<FString> LayoutNames = GetClonerLayoutNames();
+	if (!LayoutNames.Contains(InLayoutName))
+	{
+		return;
+	}
+
+	LayoutName = InLayoutName;
+	OnLayoutNameChanged();
+}
+
+void ACEClonerActor::SetRangeEnabled(bool bInRangeEnabled)
+{
+	if (bRangeEnabled == bInRangeEnabled)
+	{
+		return;
+	}
+
+	bRangeEnabled = bInRangeEnabled;
+	OnRangeOptionsChanged();
+}
+
+void ACEClonerActor::SetRangeOffsetMin(const FVector& InRangeOffsetMin)
+{
+	if (RangeOffsetMin == InRangeOffsetMin)
+	{
+		return;
+	}
+
+	RangeOffsetMin = InRangeOffsetMin;
+	OnRangeOptionsChanged();
+}
+
+void ACEClonerActor::SetRangeOffsetMax(const FVector& InRangeOffsetMax)
+{
+	if (RangeOffsetMax == InRangeOffsetMax)
+	{
+		return;
+	}
+
+	RangeOffsetMax = InRangeOffsetMax;
+	OnRangeOptionsChanged();
+}
+
+void ACEClonerActor::SetRangeRotationMin(const FRotator& InRangeRotationMin)
+{
+	if (RangeRotationMin == InRangeRotationMin)
+	{
+		return;
+	}
+
+	RangeRotationMin = InRangeRotationMin;
+	OnRangeOptionsChanged();
+}
+
+void ACEClonerActor::SetRangeRotationMax(const FRotator& InRangeRotationMax)
+{
+	if (RangeRotationMax == InRangeRotationMax)
+	{
+		return;
+	}
+
+	RangeRotationMax = InRangeRotationMax;
+	OnRangeOptionsChanged();
+}
+
+void ACEClonerActor::SetRangeScaleUniform(bool bInRangeScaleUniform)
+{
+	if (bRangeScaleUniform == bInRangeScaleUniform)
+	{
+		return;
+	}
+
+	bRangeScaleUniform = bInRangeScaleUniform;
+	OnRangeOptionsChanged();
+}
+
+void ACEClonerActor::SetRangeScaleMin(const FVector& InRangeScaleMin)
+{
+	if (RangeScaleMin == InRangeScaleMin)
+	{
+		return;
+	}
+
+	RangeScaleMin = InRangeScaleMin;
+	OnRangeOptionsChanged();
+}
+
+void ACEClonerActor::SetRangeScaleMax(const FVector& InRangeScaleMax)
+{
+	if (RangeScaleMax == InRangeScaleMax)
+	{
+		return;
+	}
+
+	RangeScaleMax = InRangeScaleMax;
+	OnRangeOptionsChanged();
+}
+
+void ACEClonerActor::SetRangeScaleUniformMin(float InRangeScaleUniformMin)
+{
+	if (RangeScaleUniformMin == InRangeScaleUniformMin)
+	{
+		return;
+	}
+
+	RangeScaleUniformMin = InRangeScaleUniformMin;
+	OnRangeOptionsChanged();
+}
+
+void ACEClonerActor::SetRangeScaleUniformMax(float InRangeScaleUniformMax)
+{
+	if (RangeScaleUniformMax == InRangeScaleUniformMax)
+	{
+		return;
+	}
+
+	RangeScaleUniformMax = InRangeScaleUniformMax;
+	OnRangeOptionsChanged();
+}
+
+void ACEClonerActor::OnMeshRenderModeChanged()
+{
+	if (!bEnabled)
+	{
+		return;
+	}
+
+	if (const UCEClonerLayoutBase* LayoutSystem = GetActiveLayout())
+	{
+		FNiagaraUserRedirectionParameterStore& ExposedParameters = LayoutSystem->GetSystem()->GetExposedParameters();
+
+		static const FNiagaraVariable MeshModeVar(FNiagaraTypeDefinition(StaticEnum<ECEClonerMeshRenderMode>()), TEXT("MeshRenderMode"));
+		ExposedParameters.SetParameterValue<int32>(static_cast<int32>(MeshRenderMode), MeshModeVar);
+
+		RequestClonerUpdate();
+	}
+}
+
+void ACEClonerActor::OnDeltaStepChanged()
+{
+	if (!ClonerComponent || !bEnabled)
+	{
+		return;
+	}
+
+	if (bDeltaStepEnabled)
+	{
+		ClonerComponent->SetVectorParameter(TEXT("DeltaStepRotation"), FVector(DeltaStepRotation.Roll, DeltaStepRotation.Pitch, DeltaStepRotation.Yaw));
+
+		ClonerComponent->SetVectorParameter(TEXT("DeltaStepScale"), DeltaStepScale);
+	}
+	else
+	{
+		ClonerComponent->SetVectorParameter(TEXT("DeltaStepRotation"), FVector::ZeroVector);
+
+		ClonerComponent->SetVectorParameter(TEXT("DeltaStepScale"), FVector::ZeroVector);
+	}
+
+	RequestClonerUpdate();
+}
+
+void ACEClonerActor::OnSeedChanged()
+{
+	if (!ClonerComponent || !bEnabled)
+	{
+		return;
+	}
+
+	ClonerComponent->SetRandomSeedOffset(Seed);
+
+	RequestClonerUpdate();
+}
+
+void ACEClonerActor::OnProgressChanged()
+{
+	if (!ClonerComponent || !bEnabled)
+	{
+		return;
+	}
+
+	ClonerComponent->SetFloatParameter(TEXT("ParticleProgress"), Progress * (bInvertProgress ? -1 : 1));
+}
+
+void ACEClonerActor::SetSpawnLoopMode(ECEClonerSpawnLoopMode InMode)
+{
+	if (SpawnLoopMode == InMode)
+	{
+		return;
+	}
+
+	SpawnLoopMode = InMode;
+	OnSpawnOptionsChanged();
+}
+
+void ACEClonerActor::SetSpawnLoopIterations(int32 InIterations)
+{
+	if (SpawnLoopIterations == InIterations)
+	{
+		return;
+	}
+
+	if (InIterations < 1)
+	{
+		return;
+	}
+
+	SpawnLoopIterations = InIterations;
+	OnSpawnOptionsChanged();
+}
+
+void ACEClonerActor::SetSpawnLoopInterval(float InInterval)
+{
+	if (SpawnLoopInterval == InInterval)
+	{
+		return;
+	}
+
+	if (InInterval < 0.f)
+	{
+		return;
+	}
+
+	SpawnLoopInterval = InInterval;
+	OnSpawnOptionsChanged();
+}
+
+void ACEClonerActor::SetSpawnBehaviorMode(ECEClonerSpawnBehaviorMode InMode)
+{
+	if (SpawnBehaviorMode == InMode)
+	{
+		return;
+	}
+
+	SpawnBehaviorMode = InMode;
+	OnSpawnOptionsChanged();
+}
+
+void ACEClonerActor::SetSpawnRate(float InRate)
+{
+	if (SpawnRate == InRate)
+	{
+		return;
+	}
+
+	if (InRate < 0)
+	{
+		return;
+	}
+
+	SpawnRate = InRate;
+	OnSpawnOptionsChanged();
+}
+
+void ACEClonerActor::SetLifetimeEnabled(bool bInEnabled)
+{
+	if (bLifetimeEnabled == bInEnabled)
+	{
+		return;
+	}
+
+	bLifetimeEnabled = bInEnabled;
+	OnLifetimeOptionsChanged();
+}
+
+void ACEClonerActor::SetLifetimeMin(float InMin)
+{
+	if (LifetimeMin == InMin)
+	{
+		return;
+	}
+
+	if (InMin < 0)
+	{
+		return;
+	}
+
+	LifetimeMin = InMin;
+	OnLifetimeOptionsChanged();
+}
+
+void ACEClonerActor::SetLifetimeMax(float InMax)
+{
+	if (LifetimeMax == InMax)
+	{
+		return;
+	}
+
+	if (InMax < 0)
+	{
+		return;
+	}
+
+	LifetimeMax = InMax;
+	OnLifetimeOptionsChanged();
+}
+
+void ACEClonerActor::SetLifetimeScaleEnabled(bool bInEnabled)
+{
+	if (bLifetimeScaleEnabled == bInEnabled)
+	{
+		return;
+	}
+
+	bLifetimeScaleEnabled = bInEnabled;
+	OnLifetimeOptionsChanged();
+}
+
+void ACEClonerActor::SetLifetimeScaleCurve(const FRichCurve& InCurve)
+{
+	LifetimeScaleCurve = InCurve;
+	OnLifetimeOptionsChanged();
+}
+
+UCEClonerLayoutBase* ACEClonerActor::GetActiveLayout() const
+{
+	if (ClonerComponent)
+	{
+		return ClonerComponent->GetClonerActiveLayout();
+	}
+
+	return nullptr;
+}
+
+const FCEClonerEffectorDataInterfaces* ACEClonerActor::GetEffectorDataInterfaces() const
+{
+	if (const UCEClonerLayoutBase* LayoutSystem = GetActiveLayout())
+	{
+		return &LayoutSystem->GetDataInterfaces();
+	}
+
+	return nullptr;
+}
+
+int32 ACEClonerActor::GetMeshCount() const
+{
+	if (ClonerComponent)
+	{
+		if (const UCEClonerLayoutBase* LayoutSystem = GetActiveLayout())
+		{
+			return LayoutSystem->GetMeshRenderer()->Meshes.Num();
+		}
+	}
+	return 0;
+}
+
+int32 ACEClonerActor::RegisterEffector(ACEEffectorActor* InEffector)
+{
+	if (!InEffector)
+	{
+		return INDEX_NONE;
+	}
+
+	int32 Index = Effectors.Find(InEffector);
+
+	if (Index != INDEX_NONE)
+	{
+		return Index;
+	}
+
+	Index = Effectors.Add(InEffector);
+
+	if (const UCEClonerLayoutBase* LayoutSystem = GetActiveLayout())
+	{
+		const int32 EffectorCount = GetEffectorCount();
+		LayoutSystem->GetDataInterfaces().Resize(EffectorCount);
+	}
+
+	InEffector->OnClonerLinked(this, Index);
+
+	constexpr bool bImmediateUpdate = true;
+	RequestClonerUpdate(bImmediateUpdate);
+
+	return Index;
+}
+
+bool ACEClonerActor::UnregisterEffector(ACEEffectorActor* InEffector)
+{
+	if (!InEffector)
+	{
+		return false;
+	}
+
+	const int32 OldIndex = Effectors.Find(InEffector);
+
+	if (OldIndex == INDEX_NONE)
+	{
+		return false;
+	}
+
+	Effectors.RemoveAt(OldIndex);
+
+	if (const UCEClonerLayoutBase* LayoutSystem = GetActiveLayout())
+	{
+		LayoutSystem->GetDataInterfaces().Remove(OldIndex);
+	}
+
+	InEffector->OnClonerUnlinked(this, OldIndex);
+
+	constexpr bool bImmediateUpdate = true;
+	RequestClonerUpdate(bImmediateUpdate);
+
+	return true;
+}
+
+bool ACEClonerActor::IsEffectorRegistered(const ACEEffectorActor* InEffector) const
+{
+	return InEffector && Effectors.Contains(InEffector);
+}
+
+int32 ACEClonerActor::GetEffectorIndex(ACEEffectorActor* InEffector) const
+{
+	return InEffector ? Effectors.Find(InEffector) : INDEX_NONE;
+}
+
+int32 ACEClonerActor::GetEffectorCount() const
+{
+	return Effectors.Num();
+}
+
+void ACEClonerActor::ForEachEffector(TFunctionRef<bool(ACEEffectorActor*, int32)> InFunction)
+{
+	for (int32 Idx = 0; Idx < Effectors.Num(); Idx++)
+	{
+		ACEEffectorActor* Effector = Effectors[Idx].Get();
+		if (!Effector)
+		{
+			continue;
+		}
+
+		if (!InFunction(Effector, Idx))
+		{
+			return;
+		}
+	}
+}
+
+void ACEClonerActor::OnEnabledChanged()
+{
+	if (bEnabled)
+	{
+		OnLayoutNameChanged();
+	}
+	else
+	{
+		ClonerComponent->DeactivateImmediate();
+		ClonerComponent->SetAsset(nullptr);
+	}
+}
+
+void ACEClonerActor::OnLayoutNameChanged()
+{
+	if (!ClonerComponent || !bEnabled)
+	{
+		return;
+	}
+
+	UCEClonerLayoutBase* NewActiveLayout = FindOrAddLayout(LayoutName);
+
+	// Apply layout
+	if (ClonerComponent->SetClonerActiveLayout(NewActiveLayout))
+	{
+		ActiveLayout = ClonerComponent->GetClonerActiveLayout();
+		OnClonerSystemChanged();
+	}
+}
+
+void ACEClonerActor::OnSpawnOptionsChanged()
+{
+	if (!ClonerComponent || !bEnabled)
+	{
+		return;
+	}
+
+	if (const UCEClonerLayoutBase* ActiveSystem = GetActiveLayout())
+	{
+		FNiagaraUserRedirectionParameterStore& ExposedParameters = ActiveSystem->GetSystem()->GetExposedParameters();
+
+		ClonerComponent->SetFloatParameter(TEXT("SpawnLoopInterval"), SpawnLoopInterval);
+
+		ClonerComponent->SetIntParameter(TEXT("SpawnLoopIterations"), SpawnLoopIterations);
+
+		ClonerComponent->SetFloatParameter(TEXT("SpawnRate"), SpawnRate);
+
+		const FNiagaraVariable SpawnBehaviorModeVar(FNiagaraTypeDefinition(StaticEnum<ECEClonerSpawnBehaviorMode>()), TEXT("SpawnBehaviorMode"));
+		ExposedParameters.SetParameterValue<int32>(static_cast<int32>(SpawnBehaviorMode), SpawnBehaviorModeVar);
+
+		const FNiagaraVariable SpawnLoopModeVar(FNiagaraTypeDefinition(StaticEnum<ECEClonerSpawnLoopMode>()), TEXT("SpawnLoopMode"));
+		ExposedParameters.SetParameterValue<int32>(static_cast<int32>(SpawnLoopMode), SpawnLoopModeVar);
+
+		RequestClonerUpdate();
+	}
+}
+
+void ACEClonerActor::OnLifetimeScaleCurveChanged()
+{
+	if (const UNiagaraDataInterfaceCurve* LifetimeCurve = LifetimeScaleCurveDIWeak.Get())
+	{
+		LifetimeScaleCurve = LifetimeCurve->Curve;
+		OnLifetimeOptionsChanged();
+	}
+}
+
+void ACEClonerActor::OnLifetimeOptionsChanged()
+{
+	if (!ClonerComponent || !bEnabled)
+	{
+		return;
+	}
+
+	LifetimeMin = FMath::Max(0, LifetimeMin);
+	LifetimeMax = FMath::Max(LifetimeMin, LifetimeMax);
+
+	ClonerComponent->SetBoolParameter(TEXT("LifetimeEnabled"), bLifetimeEnabled);
+
+	ClonerComponent->SetFloatParameter(TEXT("LifetimeMin"), LifetimeMin);
+
+	ClonerComponent->SetFloatParameter(TEXT("LifetimeMax"), LifetimeMax);
+
+	ClonerComponent->SetBoolParameter(TEXT("LifetimeScaleEnabled"), bLifetimeEnabled && bLifetimeScaleEnabled);
+
+	if (const UCEClonerLayoutBase* ActiveSystem = GetActiveLayout())
+	{
+		const FNiagaraUserRedirectionParameterStore& ExposedParameters = ActiveSystem->GetSystem()->GetExposedParameters();
+
+		static const FNiagaraVariable LifetimeScaleCurveVar(FNiagaraTypeDefinition(UNiagaraDataInterfaceCurve::StaticClass()), TEXT("LifetimeScaleCurve"));
+
+#if WITH_EDITOR
+		if (UNiagaraDataInterfaceCurve* LifetimeCurve = LifetimeScaleCurveDIWeak.Get())
+		{
+			LifetimeCurve->OnChanged().RemoveAll(this);
+		}
+#endif
+
+		if (UNiagaraDataInterfaceCurve* LifetimeScaleCurveDI = Cast<UNiagaraDataInterfaceCurve>(ExposedParameters.GetDataInterface(LifetimeScaleCurveVar)))
+		{
+			LifetimeScaleCurveDIWeak = LifetimeScaleCurveDI;
+			LifetimeScaleCurveDI->Curve = LifetimeScaleCurve;
+
+#if WITH_EDITOR
+			LifetimeScaleCurveDI->UpdateLUT();
+			LifetimeScaleCurveDI->OnChanged().AddUObject(this, &ACEClonerActor::OnLifetimeScaleCurveChanged);
+#endif
+		}
+	}
+
+	RequestClonerUpdate();
+}
+
+void ACEClonerActor::OnClonerTransformed(USceneComponent*, EUpdateTransformFlags, ETeleportType)
+{
+	UpdateClonerEffectors();
+}
+
+void ACEClonerActor::OnClonerMeshUpdated()
+{
+	RequestClonerUpdate();
+}
+
+void ACEClonerActor::OnClonerSystemChanged()
+{
+	UpdateLayoutOptions();
+	UpdateClonerEffectors();
+}
+
+void ACEClonerActor::RequestClonerUpdate(bool bInImmediate)
+{
+	if (!bEnabled || !ClonerComponent)
+	{
+		return;
+	}
+
+	if (bInImmediate)
+	{
+		bNeedsRefresh = false;
+		ClonerComponent->RefreshUserParameters();
+	}
+	else
+	{
+		bNeedsRefresh = true;
+	}
+}
+
+void ACEClonerActor::OnRangeOptionsChanged()
+{
+	if (!ClonerComponent || !bEnabled)
+	{
+		return;
+	}
+
+	RangeScaleUniformMin = FMath::Clamp(RangeScaleUniformMin, UE_KINDA_SMALL_NUMBER, RangeScaleUniformMax);
+	RangeScaleUniformMax = FMath::Max3(RangeScaleUniformMin, RangeScaleUniformMax, UE_KINDA_SMALL_NUMBER);
+
+	RangeScaleMin.X = FMath::Clamp(RangeScaleMin.X, UE_KINDA_SMALL_NUMBER, RangeScaleMax.X);
+	RangeScaleMin.Y = FMath::Clamp(RangeScaleMin.Y, UE_KINDA_SMALL_NUMBER, RangeScaleMax.Y);
+	RangeScaleMin.Z = FMath::Clamp(RangeScaleMin.Z, UE_KINDA_SMALL_NUMBER, RangeScaleMax.Z);
+
+	RangeScaleMax.X = FMath::Max3<double>(RangeScaleMin.X, RangeScaleMax.X, UE_KINDA_SMALL_NUMBER);
+	RangeScaleMax.Y = FMath::Max3<double>(RangeScaleMin.Y, RangeScaleMax.Y, UE_KINDA_SMALL_NUMBER);
+	RangeScaleMax.Z = FMath::Max3<double>(RangeScaleMin.Z, RangeScaleMax.Z, UE_KINDA_SMALL_NUMBER);
+
+	ClonerComponent->SetBoolParameter(TEXT("RangeEnabled"), bRangeEnabled);
+
+	ClonerComponent->SetVectorParameter(TEXT("RangeOffsetMin"), RangeOffsetMin);
+
+	ClonerComponent->SetVectorParameter(TEXT("RangeOffsetMax"), RangeOffsetMax);
+
+	ClonerComponent->SetVariableQuat(TEXT("RangeRotationMin"), RangeRotationMin.Quaternion());
+
+	ClonerComponent->SetVariableQuat(TEXT("RangeRotationMax"), RangeRotationMax.Quaternion());
+
+	ClonerComponent->SetBoolParameter(TEXT("RangeScaleUniform"), bRangeScaleUniform);
+
+	ClonerComponent->SetVectorParameter(TEXT("RangeScaleMin"), RangeScaleMin);
+
+	ClonerComponent->SetVectorParameter(TEXT("RangeScaleMax"), RangeScaleMax);
+
+	ClonerComponent->SetFloatParameter(TEXT("RangeScaleUniformMin"), RangeScaleUniformMin);
+
+	ClonerComponent->SetFloatParameter(TEXT("RangeScaleUniformMax"), RangeScaleUniformMax);
+
+	RequestClonerUpdate();
+}
+
+bool ACEClonerActor::IsClonerValid() const
+{
+	return IsValid(ClonerComponent.Get()) && IsValid(ClonerComponent->GetAsset());
+}
+
+void ACEClonerActor::OnMeshRendererOptionsChanged()
+{
+	const UCEClonerLayoutBase* CurrentLayout = ClonerComponent->GetClonerActiveLayout();
+
+	if (!CurrentLayout)
+	{
+		return;
+	}
+
+	UNiagaraMeshRendererProperties* MeshRenderer = CurrentLayout->GetMeshRenderer();
+
+	if (!MeshRenderer)
+	{
+		return;
+	}
+
+	MeshRenderer->FacingMode = MeshFacingMode;
+	MeshRenderer->bCastShadows = bMeshCastShadows;
+
+	RequestClonerUpdate();
+}
+
+void ACEClonerActor::OnDefaultMeshesChanged()
+{
+	// Will force an update
+	TreeUpdateDeltaTime = TreeUpdateInterval;
+}
+
+#if WITH_EDITOR
+void ACEClonerActor::OnReduceMotionGhostingChanged()
+{
+	if (UCEClonerSubsystem* ClonerSubsystem = UCEClonerSubsystem::Get())
+	{
+		if (bReduceMotionGhosting)
+		{
+			ClonerSubsystem->EnableNoFlicker();
+		}
+		else
+		{
+			ClonerSubsystem->DisableNoFlicker();
+		}
+	}
+}
+
+void ACEClonerActor::OnCVarChanged()
+{
+	if (const UCEClonerSubsystem* ClonerSubsystem = UCEClonerSubsystem::Get())
+	{
+		bReduceMotionGhosting = ClonerSubsystem->IsNoFlickerEnabled();
+	}
+}
+
+void ACEClonerActor::ForceUpdateCloner()
+{
+	if (ClonerComponent)
+	{
+		constexpr bool bReset = true;
+		ClonerComponent->UpdateClonerAttachmentTree(bReset);
+		TreeUpdateDeltaTime = TreeUpdateInterval;
+	}
+}
+
+void ACEClonerActor::SpawnLinkedEffector()
+{
+	UWorld* ClonerWorld = GetWorld();
+	if (!ClonerWorld)
+	{
+		return;
+	}
+
+	FActorSpawnParameters Params;
+	Params.bTemporaryEditorActor = false;
+
+	const FVector ClonerLocation = GetActorLocation();
+	const FRotator ClonerRotation = GetActorRotation();
+
+	ACEEffectorActor* EffectorActor = ClonerWorld->SpawnActor<ACEEffectorActor>(ACEEffectorActor::StaticClass(), ClonerLocation, ClonerRotation, Params);
+
+	if (!EffectorActor)
+	{
+		return;
+	}
+
+	EffectorActor->LinkCloner(this);
+	FActorLabelUtilities::RenameExistingActor(EffectorActor, EffectorActor->GetDefaultActorLabel(), true);
+
+	// Set default offset for visual feedback
+	EffectorActor->SetOffset(FVector(0, 0, 100));
+}
+
+void ACEClonerActor::SpawnDefaultActorAttached()
+{
+	// Only spawn if world is valid and not a preview actor
+	UWorld* World = GetWorld();
+	if (!World || bIsEditorPreviewActor)
+	{
+		return;
+	}
+
+	// Find or load cube mesh
+	constexpr const TCHAR* DefaultStaticMeshPath = TEXT("/Script/Engine.StaticMesh'/Engine/BasicShapes/Cube.Cube'");
+	UStaticMesh* DefaultStaticMesh = LoadObject<UStaticMesh>(nullptr, DefaultStaticMeshPath);
+
+	if (!DefaultStaticMesh)
+	{
+		return;
+	}
+
+	// Spawn attached actor with same flags as this actor
+	FActorSpawnParameters SpawnParameters;
+	SpawnParameters.Owner = this;
+	SpawnParameters.ObjectFlags = GetFlags();
+	SpawnParameters.bTemporaryEditorActor = false;
+
+	if (AStaticMeshActor* DefaultActorAttached = World->SpawnActor<AStaticMeshActor>(GetActorLocation(), GetActorRotation(), SpawnParameters))
+	{
+		UStaticMeshComponent* StaticMeshComponent = DefaultActorAttached->GetStaticMeshComponent();
+		StaticMeshComponent->SetStaticMesh(DefaultStaticMesh);
+
+		DefaultActorAttached->SetMobility(EComponentMobility::Movable);
+		DefaultActorAttached->AttachToActor(this, FAttachmentTransformRules::KeepWorldTransform);
+
+		FActorLabelUtilities::SetActorLabelUnique(DefaultActorAttached, TEXT("DefaultCube"));
+	}
+}
+
+void ACEClonerActor::OnEditorSelectionChanged(UObject* InSelection)
+{
+	if (const USelection* ActorSelection = Cast<USelection>(InSelection))
+	{
+		if (ULevelEditorViewportSettings* ViewportSettings = GetMutableDefault<ULevelEditorViewportSettings>())
+		{
+			if (ActorSelection->Num() == 1 && ActorSelection->GetSelectedObject(0) == this)
+			{
+				UseSelectionOutline = ViewportSettings->bUseSelectionOutline;
+				ViewportSettings->bUseSelectionOutline = false;
+			}
+			else if (UseSelectionOutline.IsSet())
+			{
+				ViewportSettings->bUseSelectionOutline = UseSelectionOutline.GetValue();
+				UseSelectionOutline.Reset();
+			}
+		}
+	}
+}
+#endif
+
+void ACEClonerActor::OnVisualizerSpriteVisibleChanged()
+{
+#if WITH_EDITOR
+	UE::ClonerEffector::SetBillboardComponentSprite(this, TEXT("/Script/Engine.Texture2D'/ClonerEffector/Textures/T_ClonerIcon.T_ClonerIcon'"));
+	UE::ClonerEffector::SetBillboardComponentVisibility(this, bVisualizerSpriteVisible);
+#endif
+}
+
+TArray<FString> ACEClonerActor::GetClonerLayoutNames() const
+{
+	TArray<FString> LayoutNamesStrings;
+
+	if (const UCEClonerSubsystem* Subsystem = UCEClonerSubsystem::Get())
+	{
+		const TArray<FName>& LayoutNames = Subsystem->GetLayoutNames();
+		LayoutNamesStrings.Reserve(LayoutNames.Num());
+
+		Algo::Transform(LayoutNames, LayoutNamesStrings, [](const FName& InName)
+		{
+			return InName.ToString();
+		});
+	}
+
+	return LayoutNamesStrings;
+}
+
+UCEClonerLayoutBase* ACEClonerActor::FindOrAddLayout(TSubclassOf<UCEClonerLayoutBase> InClass)
+{
+	const UCEClonerSubsystem* Subsystem = UCEClonerSubsystem::Get();
+
+	if (!Subsystem)
+	{
+		return nullptr;
+	}
+
+	const FName ClassLayoutName = Subsystem->FindLayoutName(InClass);
+
+	if (ClassLayoutName.IsNone())
+	{
+		return nullptr;
+	}
+
+	return FindOrAddLayout(ClassLayoutName);
+}
+
+UCEClonerLayoutBase* ACEClonerActor::FindOrAddLayout(FName InLayoutName)
+{
+	UCEClonerSubsystem* Subsystem = UCEClonerSubsystem::Get();
+	if (!Subsystem)
+	{
+		return nullptr;
+	}
+
+	// Check cached layout instances
+	UCEClonerLayoutBase* NewActiveLayout = nullptr;
+	if (TObjectPtr<UCEClonerLayoutBase> const* LayoutInstance = LayoutInstances.Find(InLayoutName))
+	{
+		NewActiveLayout = *LayoutInstance;
+	}
+
+	// Create new layout instance and cache it
+	if (!NewActiveLayout)
+	{
+		NewActiveLayout = Subsystem->CreateNewLayout(InLayoutName, this);
+		LayoutInstances.Add(InLayoutName, NewActiveLayout);
+	}
+
+	return NewActiveLayout;
+}
+
+void ACEClonerActor::InitializeCloner()
+{
+	if (bClonerInitialized)
+	{
+		return;
+	}
+
+	OnVisualizerSpriteVisibleChanged();
+
+#if WITH_EDITOR
+	OnReduceMotionGhostingChanged();
+#endif
+
+	// For new cloner instances no need to migrate anything
+	bDeprecatedPropertiesMigrated = true;
+
+	bClonerInitialized = true;
+
+	OnLayoutNameChanged();
+}
diff --git a/Engine/Plugins/Experimental/ClonerEffector/Source/ClonerEffector/Private/Cloner/CEClonerComponent.cpp b/Engine/Plugins/Experimental/ClonerEffector/Source/ClonerEffector/Private/Cloner/CEClonerComponent.cpp
new file mode 100644
--- /dev/null
+++ b/Engine/Plugins/Experimental/ClonerEffector/Source/ClonerEffector/Private/Cloner/CEClonerComponent.cpp
@@ -0,0 +1,1397 @@
+// Copyright Epic Games, Inc. All Rights Reserved.
+
+#include "Cloner/CEClonerComponent.h"
+
+#include "Async/Async.h"
+#include "Cloner/CEClonerActor.h"
+#include "Cloner/Layouts/CEClonerLayoutBase.h"
+#include "Components/BrushComponent.h"
+#include "Components/DynamicMeshComponent.h"
+#include "Components/SkeletalMeshComponent.h"
+#include "DynamicMesh/MeshNormals.h"
+#include "DynamicMesh/MeshTransforms.h"
+#include "DynamicMesh/Operations/MergeCoincidentMeshEdges.h"
+#include "DynamicMeshEditor.h"
+#include "GeometryScript/MeshAssetFunctions.h"
+#include "GeometryScript/MeshBasicEditFunctions.h"
+#include "GeometryScript/SceneUtilityFunctions.h"
+#include "Materials/Material.h"
+#include "NiagaraMeshRendererProperties.h"
+#include "NiagaraSystem.h"
+#include "ProceduralMeshComponent.h"
+#include "StaticMeshOperations.h"
+#include "UDynamicMesh.h"
+#include "Subsystems/CEClonerSubsystem.h"
+#include "UObject/Package.h"
+
+#if WITH_EDITOR
+#include "Framework/Notifications/NotificationManager.h"
+#include "Materials/Material.h"
+#include "Misc/MessageDialog.h"
+#include "Widgets/Notifications/SNotificationList.h"
+#endif
+
+DEFINE_LOG_CATEGORY_STATIC(LogCEClonerComponent, Log, All);
+
+#define LOCTEXT_NAMESPACE "CEClonerComponent"
+
+UCEClonerComponent::FOnClonerMeshUpdated UCEClonerComponent::OnClonerMeshUpdated;
+UCEClonerComponent::FOnClonerSystemLoaded UCEClonerComponent::OnClonerSystemLoaded;
+
+UCEClonerComponent::UCEClonerComponent()
+	: UNiagaraComponent()
+{
+	CastShadow = true;
+	bReceivesDecals = true;
+
+#if WITH_EDITOR
+	// Do not show bounding box around cloner for better visibility
+	SetIsVisualizationComponent(true);
+#endif
+
+	bIsEditorOnly = false;
+
+	// Show sprite for this component to visualize it when empty
+#if WITH_EDITORONLY_DATA
+	bVisualizeComponent = true;
+#endif
+
+	if (!IsTemplate())
+	{
+		// Bind to delegate to detect material changes
+#if WITH_EDITOR
+		FCoreUObjectDelegates::OnObjectPropertyChanged.RemoveAll(this);
+		FCoreUObjectDelegates::OnObjectPropertyChanged.AddUObject(this, &UCEClonerComponent::OnActorPropertyChanged);
+#endif
+	}
+}
+
+void UCEClonerComponent::PostLoad()
+{
+	SetAsset(nullptr);
+
+	Super::PostLoad();
+}
+
+void UCEClonerComponent::UpdateClonerRenderState()
+{
+	/**
+	 * Perform a mesh update when asset is valid,
+	 * An update is not already ongoing,
+	 * Meshes are out of date after an attachment tree update,
+	 * Tree is up to date
+	 */
+	if (!GetAsset()
+		|| bClonerMeshesUpdating
+		|| !bClonerMeshesDirty
+		|| ClonerTree.Status != ECEClonerAttachmentStatus::Updated)
+	{
+		return;
+	}
+
+	UpdateDirtyMeshesAsync();
+}
+
+void UCEClonerComponent::RefreshUserParameters() const
+{
+	if (UNiagaraSystem* ActiveSystem = GetAsset())
+	{
+		FNiagaraUserRedirectionParameterStore& UserParameterStore = ActiveSystem->GetExposedParameters();
+		UserParameterStore.PostGenericEditChange();
+	}
+}
+
+void UCEClonerComponent::OnRenderStateDirty(UActorComponent& InComponent)
+{
+	const AActor* OwningActor = InComponent.GetOwner();
+	const ACEClonerActor* ClonerActor = GetOuterACEClonerActor();
+
+	if (!IsValid(OwningActor) || !IsValid(ClonerActor))
+	{
+		return;
+	}
+
+	if (OwningActor->GetWorld() != ClonerActor->GetWorld())
+	{
+		return;
+	}
+
+	/* Owning actor is updated and is attached to cloner
+		or owning actor is detached from cloner */
+	if (OwningActor->IsAttachedTo(ClonerActor)
+		|| ClonerTree.ItemAttachmentMap.Contains(OwningActor))
+	{
+		UE_LOG(LogCEClonerComponent, Log, TEXT("%s : Detected attachment change for %s"), *ClonerActor->GetActorNameOrLabel(), *OwningActor->GetActorNameOrLabel());
+
+		// Only update attachment tree
+		UpdateClonerAttachmentTree();
+	}
+}
+
+void UCEClonerComponent::UpdateClonerAttachmentTree(bool bInReset)
+{
+	if (ClonerTree.Status == ECEClonerAttachmentStatus::Updated)
+	{
+		ClonerTree.Status = ECEClonerAttachmentStatus::Outdated;
+	}
+
+	if (bInReset)
+	{
+		ClonerTree.Reset();
+		ClonerTree.Status = ECEClonerAttachmentStatus::Outdated;
+	}
+
+	UpdateAttachmentTree();
+}
+
+void UCEClonerComponent::UpdateAttachmentTree()
+{
+	if (ClonerTree.Status != ECEClonerAttachmentStatus::Outdated)
+	{
+		return;
+	}
+	ClonerTree.Status = ECEClonerAttachmentStatus::Updating;
+
+	// Invalidate all, to see what is outdated and what is still invalid
+	for (TPair<TWeakObjectPtr<AActor>, FCEClonerAttachmentItem>& AttachmentPair : ClonerTree.ItemAttachmentMap)
+	{
+		AttachmentPair.Value.Status = ECEClonerAttachmentStatus::Invalid;
+	}
+
+	// Update root attachment items
+	TArray<AActor*> RootChildren;
+	GetOrderedRootActors(RootChildren);
+
+	TArray<TObjectPtr<UStaticMesh>> NewCombinesMeshes;
+	TArray<TWeakObjectPtr<AActor>> NewRootActors;
+	for (int32 RootIdx = 0; RootIdx < RootChildren.Num(); RootIdx++)
+	{
+		AActor* RootChild = RootChildren[RootIdx];
+
+		UpdateActorAttachment(RootChild, nullptr);
+
+		// Lets find the old root idx
+		const int32 OldIdx = ClonerTree.RootActors.Find(RootChild);
+		TObjectPtr<UStaticMesh> CombineBakedMesh = nullptr;
+		if (OldIdx != INDEX_NONE)
+		{
+			CombineBakedMesh = ClonerTree.MergedBakedMeshes[OldIdx];
+
+			// Did we rearrange stuff ?
+			if (RootIdx != OldIdx)
+			{
+				bClonerMeshesDirty = true;
+			}
+		}
+		NewCombinesMeshes.Add(CombineBakedMesh);
+		NewRootActors.Add(RootChild);
+	}
+
+	// Did we remove any root actors ?
+	if (ClonerTree.RootActors.Num() != NewRootActors.Num())
+	{
+		bClonerMeshesDirty = true;
+	}
+
+	// Did we need to update meshes
+	TArray<TWeakObjectPtr<AActor>> ClonedActors;
+	ClonerTree.ItemAttachmentMap.GenerateKeyArray(ClonedActors);
+	for (const TWeakObjectPtr<AActor>& ClonedActor : ClonedActors)
+	{
+		FCEClonerAttachmentItem* ClonedItem = ClonerTree.ItemAttachmentMap.Find(ClonedActor);
+
+		if (!ClonedItem)
+		{
+			continue;
+		}
+
+		if (ClonedItem->Status == ECEClonerAttachmentStatus::Invalid)
+		{
+			InvalidateBakedStaticMesh(ClonedActor.Get());
+			UnbindActorDelegates(ClonedActor.Get());
+			ClonerTree.ItemAttachmentMap.Remove(ClonedItem->ItemActor);
+		}
+		else if (ClonedItem->Status == ECEClonerAttachmentStatus::Outdated)
+		{
+			if (ClonedItem->MeshStatus == ECEClonerAttachmentStatus::Outdated)
+			{
+				DirtyItemAttachments.Add(ClonedItem);
+				InvalidateBakedStaticMesh(ClonedActor.Get());
+			}
+			bClonerMeshesDirty = true;
+			ClonedItem->Status = ECEClonerAttachmentStatus::Updated;
+		}
+	}
+
+	// Did we remove an attachment ?
+	if (ClonedActors.Num() != ClonerTree.ItemAttachmentMap.Num())
+	{
+		bClonerMeshesDirty = true;
+	}
+
+	if (!DirtyItemAttachments.IsEmpty())
+	{
+		bClonerMeshesDirty = true;
+	}
+
+	// Use default meshes if nothing is attached
+	if (NewRootActors.IsEmpty() && NewCombinesMeshes.IsEmpty())
+	{
+		if (const ACEClonerActor* ClonerActor = GetOuterACEClonerActor())
+		{
+			const TArray<TObjectPtr<UStaticMesh>>& DefaultMeshes = ClonerActor->GetDefaultMeshes();
+
+			if (ClonerTree.MergedBakedMeshes.Num() != DefaultMeshes.Num())
+			{
+				bClonerMeshesDirty = true;
+			}
+
+			for (const TObjectPtr<UStaticMesh>& DefaultMesh : DefaultMeshes)
+			{
+				const int32 Idx = NewCombinesMeshes.Add(DefaultMesh);
+
+				// Only update when there is a change
+				if (!ClonerTree.MergedBakedMeshes.IsValidIndex(Idx)
+					|| DefaultMesh != ClonerTree.MergedBakedMeshes[Idx])
+				{
+					bClonerMeshesDirty = true;
+				}
+			}
+		}
+	}
+
+	ClonerTree.RootActors = NewRootActors;
+	ClonerTree.MergedBakedMeshes = NewCombinesMeshes;
+	ClonerTree.Status = ECEClonerAttachmentStatus::Updated;
+}
+
+void UCEClonerComponent::UpdateActorAttachment(AActor* InActor, AActor* InParent)
+{
+	if (!InActor)
+	{
+		return;
+	}
+
+	const AActor* ClonerActor = GetOwner();
+	const FTransform& ClonerTransform = ClonerActor->GetActorTransform();
+
+	// Here order is not important
+	TArray<AActor*> ChildrenActors;
+	InActor->GetAttachedActors(ChildrenActors, true, false);
+
+	FCEClonerAttachmentItem* AttachmentItem = ClonerTree.ItemAttachmentMap.Find(InActor);
+	if (AttachmentItem)
+	{
+		AttachmentItem->Status = ECEClonerAttachmentStatus::Updated;
+
+		// Check Root is the same
+		const bool bIsRoot = InParent == nullptr;
+		if (AttachmentItem->bRootItem != bIsRoot)
+		{
+			InvalidateBakedStaticMesh(InActor);
+			AttachmentItem->bRootItem = bIsRoot;
+			AttachmentItem->Status = ECEClonerAttachmentStatus::Outdated;
+		}
+
+		// Check parent is the same
+		if (AttachmentItem->ParentActor.Get() != InParent)
+		{
+			InvalidateBakedStaticMesh(InParent);
+			InvalidateBakedStaticMesh(AttachmentItem->ParentActor.Get());
+			AttachmentItem->ParentActor = InParent;
+			AttachmentItem->Status = ECEClonerAttachmentStatus::Outdated;
+		}
+
+		// Check transform is the same
+		const FTransform ActorTransform = InActor->GetActorTransform().GetRelativeTransform(ClonerTransform);
+		if (!ActorTransform.Equals(AttachmentItem->ActorTransform))
+		{
+			// invalidate if not root, else change transform in mesh renderer
+			if (!bIsRoot)
+			{
+				InvalidateBakedStaticMesh(InActor);
+			}
+			AttachmentItem->ActorTransform = ActorTransform;
+			AttachmentItem->Status = ECEClonerAttachmentStatus::Outdated;
+		}
+	}
+	else
+	{
+		AttachmentItem = &ClonerTree.ItemAttachmentMap.Add(InActor, FCEClonerAttachmentItem());
+		AttachmentItem->ItemActor = InActor;
+		AttachmentItem->ParentActor = InParent;
+		AttachmentItem->ActorTransform = InActor->GetActorTransform().GetRelativeTransform(ClonerTransform);
+		AttachmentItem->MeshStatus = ECEClonerAttachmentStatus::Outdated;
+		AttachmentItem->bRootItem = InParent == nullptr;
+		AttachmentItem->Status = ECEClonerAttachmentStatus::Outdated;
+		InvalidateBakedStaticMesh(InActor);
+		BindActorDelegates(InActor);
+	}
+
+	if (AttachmentItem->ChildrenActors.Num() != ChildrenActors.Num())
+	{
+		InvalidateBakedStaticMesh(InActor);
+	}
+
+	AttachmentItem->ChildrenActors.Empty(ChildrenActors.Num());
+	for (AActor* ChildActor : ChildrenActors)
+	{
+		AttachmentItem->ChildrenActors.Add(ChildActor);
+		UpdateActorAttachment(ChildActor, InActor);
+	}
+}
+
+void UCEClonerComponent::BindActorDelegates(AActor* InActor)
+{
+	if (!InActor)
+	{
+		return;
+	}
+
+#if WITH_EDITOR
+	// Hide new actor when attached to cloner
+	InActor->SetIsTemporarilyHiddenInEditor(true);
+#endif
+	InActor->SetActorHiddenInGame(true);
+
+	InActor->OnDestroyed.RemoveAll(this);
+	InActor->OnDestroyed.AddDynamic(this, &UCEClonerComponent::OnActorDestroyed);
+
+#if WITH_EDITOR
+	// Detect static mesh change
+	TArray<UStaticMeshComponent*> StaticMeshComponents;
+	InActor->GetComponents(StaticMeshComponents, false);
+	for (UStaticMeshComponent* StaticMeshComponent : StaticMeshComponents)
+	{
+		if (!StaticMeshComponent->OnStaticMeshChanged().IsBoundToObject(this))
+		{
+			StaticMeshComponent->OnStaticMeshChanged().AddUObject(this, &UCEClonerComponent::OnMeshChanged, InActor);
+		}
+	}
+#endif
+
+	// Detect dynamic mesh change
+	TArray<UDynamicMeshComponent*> DynamicMeshComponents;
+	InActor->GetComponents(DynamicMeshComponents, false);
+	for (UDynamicMeshComponent* DynamicMeshComponent : DynamicMeshComponents)
+	{
+		if (!DynamicMeshComponent->OnMeshChanged.IsBoundToObject(this))
+		{
+			UStaticMeshComponent* NullComponent = nullptr;
+			DynamicMeshComponent->OnMeshChanged.AddUObject(this, &UCEClonerComponent::OnMeshChanged, NullComponent, InActor);
+		}
+	}
+}
+
+void UCEClonerComponent::UnbindActorDelegates(AActor* InActor) const
+{
+	if (!InActor)
+	{
+		return;
+	}
+
+#if WITH_EDITOR
+	// Show new actor when detached to cloner
+	InActor->SetIsTemporarilyHiddenInEditor(false);
+#endif
+	InActor->SetActorHiddenInGame(false);
+
+	InActor->OnDestroyed.RemoveAll(this);
+
+#if WITH_EDITOR
+	TArray<UStaticMeshComponent*> StaticMeshComponents;
+	InActor->GetComponents(StaticMeshComponents, false);
+	for (UStaticMeshComponent* StaticMeshComponent : StaticMeshComponents)
+	{
+		StaticMeshComponent->OnStaticMeshChanged().RemoveAll(this);
+	}
+#endif
+
+	TArray<UDynamicMeshComponent*> DynamicMeshComponents;
+	InActor->GetComponents(DynamicMeshComponents, false);
+	for (UDynamicMeshComponent* DynamicMeshComponent : DynamicMeshComponents)
+	{
+		DynamicMeshComponent->OnMeshChanged.RemoveAll(this);
+	}
+}
+
+void UCEClonerComponent::OnActorDestroyed(AActor* InDestroyedActor)
+{
+	if (ClonerTree.ItemAttachmentMap.Contains(InDestroyedActor))
+	{
+		InvalidateBakedStaticMesh(InDestroyedActor);
+		UnbindActorDelegates(InDestroyedActor);
+		ClonerTree.ItemAttachmentMap.Remove(InDestroyedActor);
+		bClonerMeshesDirty = true;
+	}
+}
+
+#if WITH_EDITOR
+void UCEClonerComponent::OnActorPropertyChanged(UObject* InObject, FPropertyChangedEvent& InPropertyChangedEvent)
+{
+	OnMaterialChanged(InObject);
+}
+#endif
+
+void UCEClonerComponent::OnMaterialChanged(UObject* InObject)
+{
+	const AActor* ClonerActor = GetOwner();
+
+	if (!ClonerActor)
+	{
+		return;
+	}
+
+	if (AActor* ActorChanged = Cast<AActor>(InObject))
+	{
+		if (FCEClonerAttachmentItem* Item = ClonerTree.ItemAttachmentMap.Find(ActorChanged))
+		{
+			int32 MatIdx = 0;
+			TArray<TWeakObjectPtr<UMaterialInterface>> NewMaterials;
+			bool bMaterialChanged = false;
+
+			TArray<UDynamicMeshComponent*> DynamicMeshComponents;
+			ActorChanged->GetComponents(DynamicMeshComponents, false);
+			for (const UDynamicMeshComponent* DynamicMeshComponent : DynamicMeshComponents)
+			{
+				for (UMaterialInterface* Material : DynamicMeshComponent->GetMaterials())
+				{
+					if (!Item->BakedMaterials.IsValidIndex(MatIdx) || Item->BakedMaterials[MatIdx] != Material)
+					{
+						bMaterialChanged = true;
+					}
+					NewMaterials.Add(Material);
+					MatIdx++;
+				}
+			}
+
+			TArray<USkeletalMeshComponent*> SkeletalMeshComponents;
+			ActorChanged->GetComponents(SkeletalMeshComponents, false);
+			for (const USkeletalMeshComponent* SkeletalMeshComponent : SkeletalMeshComponents)
+			{
+				for (UMaterialInterface* Material : SkeletalMeshComponent->GetMaterials())
+				{
+					if (!Item->BakedMaterials.IsValidIndex(MatIdx) || Item->BakedMaterials[MatIdx] != Material)
+					{
+						bMaterialChanged = true;
+					}
+					NewMaterials.Add(Material);
+					MatIdx++;
+				}
+			}
+
+			TArray<UBrushComponent*> BrushComponents;
+			ActorChanged->GetComponents(BrushComponents, false);
+			for (const UBrushComponent* BrushComponent : BrushComponents)
+			{
+				for (int32 Idx = 0; Idx < BrushComponent->GetNumMaterials(); Idx++)
+				{
+					if (!Item->BakedMaterials.IsValidIndex(MatIdx) || Item->BakedMaterials[MatIdx] != BrushComponent->GetMaterial(Idx))
+					{
+						bMaterialChanged = true;
+					}
+					NewMaterials.Add(BrushComponent->GetMaterial(Idx));
+					MatIdx++;
+				}
+			}
+
+			TArray<UProceduralMeshComponent*> ProceduralMeshComponents;
+			ActorChanged->GetComponents(ProceduralMeshComponents, false);
+			for (const UProceduralMeshComponent* ProceduralMeshComponent : ProceduralMeshComponents)
+			{
+				if (ProceduralMeshComponent->GetNumSections() == 0)
+				{
+					continue;
+				}
+				for (UMaterialInterface* Material : ProceduralMeshComponent->GetMaterials())
+				{
+					if (!Item->BakedMaterials.IsValidIndex(MatIdx) || Item->BakedMaterials[MatIdx] != Material)
+					{
+						bMaterialChanged = true;
+					}
+					NewMaterials.Add(Material);
+					MatIdx++;
+				}
+			}
+
+			TArray<UStaticMeshComponent*> StaticMeshComponents;
+			ActorChanged->GetComponents(StaticMeshComponents, false);
+			for (const UStaticMeshComponent* StaticMeshComponent : StaticMeshComponents)
+			{
+				for (UMaterialInterface* Material : StaticMeshComponent->GetMaterials())
+				{
+					if (!Item->BakedMaterials.IsValidIndex(MatIdx) || Item->BakedMaterials[MatIdx] != Material)
+					{
+						bMaterialChanged = true;
+					}
+					NewMaterials.Add(Material);
+					MatIdx++;
+				}
+			}
+
+			if (bMaterialChanged)
+			{
+				UE_LOG(LogCEClonerComponent, Log, TEXT("%s : Detected material change for %s"), *ClonerActor->GetActorNameOrLabel(), *ActorChanged->GetActorNameOrLabel());
+
+				if (NewMaterials.Num() == Item->BakedMaterials.Num())
+				{
+					Item->BakedMaterials = NewMaterials;
+				}
+				else
+				{
+					Item->MeshStatus = ECEClonerAttachmentStatus::Outdated;
+				}
+
+				InvalidateBakedStaticMesh(ActorChanged);
+			}
+		}
+	}
+}
+
+void UCEClonerComponent::OnTransformUpdated(USceneComponent* InUpdatedComponent, EUpdateTransformFlags InUpdateTransformFlags, ETeleportType InTeleport)
+{
+	InvalidateBakedStaticMesh(InUpdatedComponent->GetOwner());
+}
+
+void UCEClonerComponent::OnMeshChanged(UStaticMeshComponent*, AActor* InActor)
+{
+	if (!InActor)
+	{
+		return;
+	}
+
+	const AActor* ClonerActor = GetOwner();
+
+	if (!ClonerActor)
+	{
+		return;
+	}
+
+	if (FCEClonerAttachmentItem* Item = ClonerTree.ItemAttachmentMap.Find(InActor))
+	{
+		UE_LOG(LogCEClonerComponent, Log, TEXT("%s : Detected mesh change for %s"), *ClonerActor->GetActorNameOrLabel(), *InActor->GetActorNameOrLabel());
+
+		Item->MeshStatus = ECEClonerAttachmentStatus::Outdated;
+        InvalidateBakedStaticMesh(InActor);
+		DirtyItemAttachments.Add(Item);
+	}
+}
+
+void UCEClonerComponent::GetOrderedRootActors(TArray<AActor*>& OutActors) const
+{
+	const AActor* ClonerActor = GetOwner();
+	if (!ClonerActor)
+	{
+		return;
+	}
+
+	UCEClonerSubsystem* ClonerSubsystem = UCEClonerSubsystem::Get();
+	if (!ClonerSubsystem)
+	{
+		return;
+	}
+
+	const UCEClonerSubsystem::FOnGetOrderedActors& CustomActorResolver = ClonerSubsystem->GetCustomActorResolver();
+
+	if (CustomActorResolver.IsBound())
+	{
+		OutActors = CustomActorResolver.Execute(ClonerActor);
+	}
+	else
+	{
+		ClonerActor->GetAttachedActors(OutActors, true, false);
+	}
+}
+
+AActor* UCEClonerComponent::GetRootActor(AActor* InActor) const
+{
+	if (InActor == nullptr)
+	{
+		return nullptr;
+	}
+	if (const FCEClonerAttachmentItem* Item = ClonerTree.ItemAttachmentMap.Find(InActor))
+	{
+		if (Item->bRootItem)
+		{
+			return InActor;
+		}
+
+		return GetRootActor(Item->ParentActor.Get());
+	}
+	return nullptr;
+}
+
+void UCEClonerComponent::InvalidateBakedStaticMesh(AActor* InActor)
+{
+	if (!InActor)
+	{
+		return;
+	}
+
+	if (const FCEClonerAttachmentItem* FoundItem = ClonerTree.ItemAttachmentMap.Find(InActor))
+	{
+		if (FoundItem->bRootItem || !FoundItem->ParentActor.IsValid())
+		{
+			const int32 RootIdx = ClonerTree.RootActors.Find(InActor);
+			if (ClonerTree.MergedBakedMeshes.IsValidIndex(RootIdx))
+			{
+				ClonerTree.MergedBakedMeshes[RootIdx] = nullptr;
+				bClonerMeshesDirty = true;
+			}
+		}
+		else
+		{
+			InvalidateBakedStaticMesh(FoundItem->ParentActor.Get());
+		}
+	}
+}
+
+void UCEClonerComponent::UpdateDirtyMeshesAsync()
+{
+	if (bClonerMeshesUpdating)
+	{
+		return;
+	}
+
+	bClonerMeshesUpdating = true;
+	TSet<FCEClonerAttachmentItem*> UpdateItems = DirtyItemAttachments;
+	DirtyItemAttachments.Empty();
+
+	// Update baked dynamic meshes on other thread
+	TWeakObjectPtr<UCEClonerComponent> ThisWeak(this);
+	Async(EAsyncExecution::TaskGraph, [ThisWeak, UpdateItems]()->bool
+	{
+		if (!ThisWeak.IsValid())
+		{
+			return false;
+		}
+
+		// update actor baked dynamic meshes
+		for (const FCEClonerAttachmentItem* Item : UpdateItems)
+		{
+			if (!Item || !Item->ItemActor.IsValid())
+			{
+				continue;
+			}
+
+			ThisWeak->UpdateActorBakedDynamicMesh(Item->ItemActor.Get());
+		}
+
+		// Create baked static mesh on main thread (required)
+		Async(EAsyncExecution::TaskGraphMainThread, [ThisWeak]()
+		{
+			if (!ThisWeak.IsValid())
+			{
+				return;
+			}
+
+			// Update actors baked static mesh
+			for (int32 Idx = 0; Idx < ThisWeak->ClonerTree.RootActors.Num(); Idx++)
+			{
+				AActor* RootActor = ThisWeak->ClonerTree.RootActors[Idx].Get();
+				if (!ThisWeak->ClonerTree.MergedBakedMeshes[Idx].Get())
+				{
+					ThisWeak->UpdateRootActorBakedStaticMesh(RootActor);
+				}
+			}
+
+			// update niagara asset
+			ThisWeak->OnDirtyMeshesUpdated();
+		});
+
+		return true;
+	});
+}
+
+void UCEClonerComponent::OnDirtyMeshesUpdated()
+{
+	bClonerMeshesUpdating = false;
+
+	// Update niagara parameters
+	UpdateClonerMeshes();
+}
+
+void UCEClonerComponent::UpdateActorBakedDynamicMesh(AActor* InActor)
+{
+	if (!InActor)
+	{
+		return;
+	}
+
+	FCEClonerAttachmentItem* AttachmentItem = ClonerTree.ItemAttachmentMap.Find(InActor);
+
+	if (!AttachmentItem || AttachmentItem->MeshStatus != ECEClonerAttachmentStatus::Outdated)
+	{
+		return;
+	}
+
+	AActor* ClonerActor = GetOwner();
+
+	if (!ClonerActor)
+	{
+		return;
+	}
+
+	AttachmentItem->MeshStatus = ECEClonerAttachmentStatus::Updating;
+
+	UE_LOG(LogCEClonerComponent, Log, TEXT("%s : Updating baked actor mesh %s"), *ClonerActor->GetActorNameOrLabel(), *InActor->GetActorNameOrLabel());
+
+	using namespace UE::Geometry;
+
+	TArray<FDynamicMesh3> ConvertedMeshes;
+	TArray<TWeakObjectPtr<UMaterialInterface>> MeshesMaterials;
+
+	const FTransform& SourceTransform = InActor->GetActorTransform();
+	UDynamicMesh* OutputDynamicMesh = NewObject<UDynamicMesh>();
+
+	static FGeometryScriptCopyMeshFromAssetOptions OutputMeshOptions;
+	OutputMeshOptions.bIgnoreRemoveDegenerates = false;
+	OutputMeshOptions.bRequestTangents = false;
+	OutputMeshOptions.bApplyBuildSettings = false;
+
+	// Dynamic mesh components
+	{
+		TArray<UDynamicMeshComponent*> Components;
+		InActor->GetComponents(Components, false);
+		for (UDynamicMeshComponent* Component : Components)
+		{
+			const UDynamicMesh* DynamicMesh = Component->GetDynamicMesh();
+			// Transform the new mesh relative to the component
+			const FTransform RelativeTransform = Component->GetComponentTransform().GetRelativeTransform(SourceTransform);
+			MeshesMaterials.Append(Component->GetMaterials());
+			// Create a copy
+			DynamicMesh->ProcessMesh([&ConvertedMeshes, RelativeTransform](const FDynamicMesh3& EditMesh)
+			{
+				FDynamicMesh3 CopyMesh = EditMesh;
+				MeshTransforms::ApplyTransform(CopyMesh, RelativeTransform);
+				ConvertedMeshes.Add(MoveTemp(CopyMesh));
+			});
+		}
+	}
+
+	// Skeletal mesh components
+	{
+		static FGeometryScriptMeshReadLOD SkeletalMeshLOD;
+		SkeletalMeshLOD.LODType = EGeometryScriptLODType::SourceModel;
+
+		TArray<USkeletalMeshComponent*> Components;
+		InActor->GetComponents(Components, false);
+		for (USkeletalMeshComponent* Component : Components)
+		{
+			USkeletalMesh* SkeletalMesh = Component->GetSkeletalMeshAsset();
+
+			// convert to dynamic mesh
+			EGeometryScriptOutcomePins OutResult;
+			UGeometryScriptLibrary_StaticMeshFunctions::CopyMeshFromSkeletalMesh(SkeletalMesh, OutputDynamicMesh, OutputMeshOptions, SkeletalMeshLOD, OutResult);
+			if (OutResult == EGeometryScriptOutcomePins::Success)
+			{
+				// Transform the new mesh relative to the component
+				const FTransform RelativeTransform = Component->GetComponentTransform().GetRelativeTransform(SourceTransform);
+				MeshesMaterials.Append(Component->GetMaterials());
+				OutputDynamicMesh->EditMesh([&ConvertedMeshes, RelativeTransform](FDynamicMesh3& EditMesh)
+				{
+					MeshTransforms::ApplyTransform(EditMesh, RelativeTransform);
+					ConvertedMeshes.Add(MoveTemp(EditMesh));
+					// replace by empty mesh
+					FDynamicMesh3 EmptyMesh;
+					EditMesh = MoveTemp(EmptyMesh);
+				}, EDynamicMeshChangeType::GeneralEdit, EDynamicMeshAttributeChangeFlags::Unknown, true);
+			}
+		}
+	}
+
+	// Brush components
+	{
+		static const FGeometryScriptCopyMeshFromComponentOptions Options;
+		static FTransform Transform;
+
+		TArray<UBrushComponent*> Components;
+		InActor->GetComponents(Components, false);
+		for (UBrushComponent* Component : Components)
+		{
+			// convert to dynamic mesh
+			EGeometryScriptOutcomePins OutResult;
+			UGeometryScriptLibrary_SceneUtilityFunctions::CopyMeshFromComponent(Component, OutputDynamicMesh, Options, false, Transform, OutResult);
+			if (OutResult == EGeometryScriptOutcomePins::Success)
+			{
+				// Transform the new mesh relative to the component
+				const FTransform RelativeTransform = Component->GetComponentTransform().GetRelativeTransform(SourceTransform);
+				for (int32 MatIdx = 0; MatIdx < Component->GetNumMaterials(); MatIdx++)
+				{
+					MeshesMaterials.Add(Component->GetMaterial(MatIdx));
+				}
+				OutputDynamicMesh->EditMesh([&ConvertedMeshes, RelativeTransform](FDynamicMesh3& EditMesh)
+				{
+					MeshTransforms::ApplyTransform(EditMesh, RelativeTransform);
+					ConvertedMeshes.Add(MoveTemp(EditMesh));
+					// replace by empty mesh
+					FDynamicMesh3 EmptyMesh;
+					EditMesh = MoveTemp(EmptyMesh);
+				}, EDynamicMeshChangeType::GeneralEdit, EDynamicMeshAttributeChangeFlags::Unknown, true);
+			}
+		}
+	}
+
+	// Procedural mesh components
+	{
+		TArray<UProceduralMeshComponent*> Components;
+		InActor->GetComponents(Components, false);
+		for (UProceduralMeshComponent* Component : Components)
+		{
+			const int32 SectionCount = Component->GetNumSections();
+			if (SectionCount == 0)
+			{
+				continue;
+			}
+
+			// Transform the new mesh relative to the component
+			const FTransform RelativeTransform = Component->GetComponentTransform().GetRelativeTransform(SourceTransform);
+
+			FDynamicMesh3 Mesh;
+			Mesh.EnableAttributes();
+			Mesh.Attributes()->EnablePrimaryColors();
+			Mesh.Attributes()->EnableMaterialID();
+			Mesh.Attributes()->SetNumNormalLayers(1);
+			Mesh.Attributes()->SetNumUVLayers(1);
+			Mesh.Attributes()->SetNumPolygroupLayers(1);
+			Mesh.Attributes()->EnableTangents();
+
+			FDynamicMeshColorOverlay* ColorOverlay = Mesh.Attributes()->PrimaryColors();
+			FDynamicMeshNormalOverlay* NormalOverlay = Mesh.Attributes()->PrimaryNormals();
+			FDynamicMeshUVOverlay* UVOverlay = Mesh.Attributes()->PrimaryUV();
+			FDynamicMeshMaterialAttribute* MaterialAttr = Mesh.Attributes()->GetMaterialID();
+			FDynamicMeshPolygroupAttribute* GroupAttr = Mesh.Attributes()->GetPolygroupLayer(0);
+			FDynamicMeshNormalOverlay* TangentOverlay = Mesh.Attributes()->PrimaryTangents();
+
+			for (int32 SectionIdx = 0; SectionIdx < SectionCount; SectionIdx++)
+			{
+				if (FProcMeshSection* Section = Component->GetProcMeshSection(SectionIdx))
+				{
+					if (Section->bSectionVisible)
+					{
+						TArray<int32> VtxIds;
+						TArray<int32> NormalIds;
+						TArray<int32> ColorIds;
+						TArray<int32> UVIds;
+						TArray<int32> TaIds;
+
+						// copy vertices data (position, normal, color, UV, tangent)
+						for (FProcMeshVertex& SectionVertex : Section->ProcVertexBuffer)
+						{
+							int32 VId = Mesh.AppendVertex(SectionVertex.Position);
+							VtxIds.Add(VId);
+
+							int32 NId = NormalOverlay->AppendElement(static_cast<FVector3f>(SectionVertex.Normal));
+							NormalIds.Add(NId);
+
+							int32 CId = ColorOverlay->AppendElement(static_cast<FVector4f>(SectionVertex.Color));
+							ColorIds.Add(CId);
+
+							int32 UVId = UVOverlay->AppendElement(static_cast<FVector2f>(SectionVertex.UV0));
+							UVIds.Add(UVId);
+
+							int32 TaId = TangentOverlay->AppendElement(static_cast<FVector3f>(SectionVertex.Tangent.TangentX));
+							TaIds.Add(TaId);
+						}
+
+						// copy tris data
+						if (Section->ProcIndexBuffer.Num() % 3 != 0)
+						{
+							continue;
+						}
+						for (int32 Idx = 0; Idx < Section->ProcIndexBuffer.Num(); Idx+=3)
+						{
+							int32 VIdx1 = Section->ProcIndexBuffer[Idx];
+							int32 VIdx2 = Section->ProcIndexBuffer[Idx + 1];
+							int32 VIdx3 = Section->ProcIndexBuffer[Idx + 2];
+
+							int32 VId1 = VtxIds[VIdx1];
+							int32 VId2 = VtxIds[VIdx2];
+							int32 VId3 = VtxIds[VIdx3];
+
+							int32 TId = Mesh.AppendTriangle(VId1, VId2, VId3, SectionIdx);
+							if (TId < 0)
+							{
+								continue;
+							}
+
+							NormalOverlay->SetTriangle(TId, FIndex3i(NormalIds[VIdx1], NormalIds[VIdx2], NormalIds[VIdx3]), true);
+							ColorOverlay->SetTriangle(TId, FIndex3i(ColorIds[VIdx1], ColorIds[VIdx2], ColorIds[VIdx3]), true);
+							UVOverlay->SetTriangle(TId, FIndex3i(UVIds[VIdx1], UVIds[VIdx2], UVIds[VIdx3]), true);
+							TangentOverlay->SetTriangle(TId, FIndex3i(TaIds[VIdx1], TaIds[VIdx2], TaIds[VIdx3]), true);
+
+							MaterialAttr->SetValue(TId, SectionIdx);
+							GroupAttr->SetValue(TId, SectionIdx);
+						}
+					}
+				}
+			}
+
+			if (Mesh.TriangleCount() > 0)
+			{
+				MeshTransforms::ApplyTransform(Mesh, RelativeTransform);
+				MeshesMaterials.Append(Component->GetMaterials());
+				ConvertedMeshes.Add(MoveTemp(Mesh));
+			}
+		}
+	}
+
+	// Static mesh components
+	{
+		static FGeometryScriptMeshReadLOD StaticMeshLOD;
+		StaticMeshLOD.LODType = EGeometryScriptLODType::RenderData;
+
+		TArray<UStaticMeshComponent*> Components;
+		InActor->GetComponents(Components, false);
+		for (UStaticMeshComponent* Component : Components)
+		{
+			UStaticMesh* StaticMesh = Component->GetStaticMesh();
+			// convert to dynamic mesh
+			EGeometryScriptOutcomePins OutResult;
+			UGeometryScriptLibrary_StaticMeshFunctions::CopyMeshFromStaticMesh(StaticMesh, OutputDynamicMesh, OutputMeshOptions, StaticMeshLOD, OutResult);
+			if (OutResult == EGeometryScriptOutcomePins::Success)
+			{
+				// Transform the new mesh relative to the component
+				const FTransform RelativeTransform = Component->GetComponentTransform().GetRelativeTransform(SourceTransform);
+				MeshesMaterials.Append(Component->GetMaterials());
+				OutputDynamicMesh->EditMesh([&ConvertedMeshes, RelativeTransform](FDynamicMesh3& EditMesh)
+				{
+					MeshTransforms::ApplyTransform(EditMesh, RelativeTransform);
+					ConvertedMeshes.Add(MoveTemp(EditMesh));
+					// replace by empty mesh
+					FDynamicMesh3 EmptyMesh;
+					EditMesh = MoveTemp(EmptyMesh);
+				}, EDynamicMeshChangeType::GeneralEdit, EDynamicMeshAttributeChangeFlags::Unknown, true);
+			}
+		}
+	}
+
+	// Lets combine all meshes components from this actor together
+	OutputDynamicMesh->EditMesh([&ConvertedMeshes](FDynamicMesh3& InMergedMesh)
+	{
+		InMergedMesh.Clear();
+		InMergedMesh.EnableAttributes();
+		InMergedMesh.Attributes()->SetNumNormalLayers(1);
+
+		FDynamicMeshEditor Editor(&InMergedMesh);
+		FGeometryScriptAppendMeshOptions AppendOptions;
+		AppendOptions.CombineMode = EGeometryScriptCombineAttributesMode::EnableAllMatching;
+		int32 MaterialCount = 0;
+
+		// Convert meshes
+		for (int32 MeshIdx = 0; MeshIdx < ConvertedMeshes.Num(); MeshIdx++)
+		{
+			const FDynamicMesh3& ConvertedMesh = ConvertedMeshes[MeshIdx];
+
+			// Enable matching attributes & append mesh
+			FMeshIndexMappings TmpMappings;
+			AppendOptions.UpdateAttributesForCombineMode(InMergedMesh, ConvertedMesh);
+			Editor.AppendMesh(&ConvertedMesh, TmpMappings);
+
+			// Fix triangles materials linking
+			if (ConvertedMesh.HasAttributes() && ConvertedMesh.Attributes()->HasMaterialID())
+			{
+				const FDynamicMeshMaterialAttribute* FromMaterialIDAttrib = ConvertedMesh.Attributes()->GetMaterialID();
+				FDynamicMeshMaterialAttribute* ToMaterialIDAttrib = InMergedMesh.Attributes()->GetMaterialID();
+				TMap<int32, int32> MaterialMap;
+				for (const TPair<int32, int32>& FromToTId : TmpMappings.GetTriangleMap().GetForwardMap())
+				{
+					const int32 FromMatId = FromMaterialIDAttrib->GetValue(FromToTId.Key);
+					const int32 ToMatId = FromMatId + MaterialCount;
+					MaterialMap.Add(FromMatId, ToMatId);
+					ToMaterialIDAttrib->SetNewValue(FromToTId.Value, ToMatId);
+				}
+				MaterialCount += MaterialMap.Num();
+			}
+		}
+
+		if (InMergedMesh.TriangleCount() > 0)
+		{
+			// Merge shared edges
+			FMergeCoincidentMeshEdges WeldOp(&InMergedMesh);
+			WeldOp.Apply();
+		}
+	}, EDynamicMeshChangeType::GeneralEdit, EDynamicMeshAttributeChangeFlags::Unknown, true);
+
+	AttachmentItem->BakedMesh = OutputDynamicMesh;
+	AttachmentItem->BakedMaterials = MeshesMaterials;
+	// Was the mesh invalidated during the update process ?
+	AttachmentItem->MeshStatus = AttachmentItem->MeshStatus == ECEClonerAttachmentStatus::Outdated ? ECEClonerAttachmentStatus::Outdated : ECEClonerAttachmentStatus::Updated;
+	InvalidateBakedStaticMesh(AttachmentItem->ItemActor.Get());
+}
+
+void UCEClonerComponent::UpdateRootActorBakedStaticMesh(AActor* InRootActor)
+{
+	if (!InRootActor)
+	{
+		return;
+	}
+
+	const int32 RootIdx = ClonerTree.RootActors.Find(InRootActor);
+
+	if (RootIdx == INDEX_NONE)
+	{
+		return;
+	}
+
+	const FCEClonerAttachmentItem* RootAttachmentItem = ClonerTree.ItemAttachmentMap.Find(InRootActor);
+
+	if (!RootAttachmentItem)
+	{
+		return;
+	}
+
+	const AActor* ClonerActor = GetOwner();
+
+	if (!ClonerActor)
+	{
+		return;
+	}
+
+	UE_LOG(LogCEClonerComponent, Log, TEXT("%s : Updating root merged baked mesh %s"), *ClonerActor->GetActorNameOrLabel(), *InRootActor->GetActorNameOrLabel());
+
+	using namespace UE::Geometry;
+
+#if WITH_EDITOR
+	TSet<int32> MaterialsMissingNiagaraUsageFlag;
+#endif
+
+	TArray<FCEClonerAttachmentItem*> AttachmentItems;
+	GetActorAttachmentItems(InRootActor, AttachmentItems);
+
+	// Lets combine all dynamic meshes from item attachments and self into one mesh
+	UDynamicMesh* MergedAttachmentMesh = NewObject<UDynamicMesh>();
+	TArray<TObjectPtr<UMaterialInterface>> CombineMaterials;
+	MergedAttachmentMesh->EditMesh([&AttachmentItems, &CombineMaterials, this
+#if WITH_EDITOR
+		, &MaterialsMissingNiagaraUsageFlag
+#endif
+		](FDynamicMesh3& InMergedMesh)
+	{
+		InMergedMesh.Clear();
+		InMergedMesh.EnableAttributes();
+		InMergedMesh.Attributes()->SetNumNormalLayers(1);
+
+		FDynamicMeshEditor Editor(&InMergedMesh);
+		FGeometryScriptAppendMeshOptions AppendOptions;
+		AppendOptions.CombineMode = EGeometryScriptCombineAttributesMode::EnableAllMatching;
+		int32 MaterialCount = 0;
+
+		// Convert meshes
+		for (int32 MeshIdx = 0; MeshIdx < AttachmentItems.Num(); MeshIdx++)
+		{
+			FCEClonerAttachmentItem* AttachmentItem = AttachmentItems[MeshIdx];
+			if (!AttachmentItem)
+			{
+				continue;
+			}
+
+			const UDynamicMesh* BakedDynamicMesh = AttachmentItem->BakedMesh;
+			if (!BakedDynamicMesh)
+			{
+				continue;
+			}
+
+			BakedDynamicMesh->ProcessMesh([&AttachmentItem, &CombineMaterials, &MaterialCount, &AppendOptions, &InMergedMesh, &Editor, this
+#if WITH_EDITOR
+				, &MaterialsMissingNiagaraUsageFlag
+#endif
+				](const FDynamicMesh3& InProcessMesh)
+			{
+				FDynamicMesh3 ConvertedMesh = InProcessMesh;
+
+				if (AttachmentItem->ParentActor.IsValid())
+				{
+					const FTransform& ParentTransform = AttachmentItem->ParentActor->GetTransform();
+					const FTransform ParentChildTransform = AttachmentItem->ItemActor->GetTransform().GetRelativeTransform(ParentTransform);
+					MeshTransforms::ApplyTransform(ConvertedMesh, ParentChildTransform);
+				}
+				else
+				{
+					MeshTransforms::ApplyTransform(ConvertedMesh, FTransform::Identity);
+				}
+
+				// Enable matching attributes & append mesh
+				FMeshIndexMappings TmpMappings;
+				AppendOptions.UpdateAttributesForCombineMode(InMergedMesh, ConvertedMesh);
+				Editor.AppendMesh(&ConvertedMesh, TmpMappings);
+
+				// Fix triangles materials linking
+				if (ConvertedMesh.HasAttributes() && ConvertedMesh.Attributes()->HasMaterialID())
+				{
+					const FDynamicMeshMaterialAttribute* FromMaterialIDAttrib = ConvertedMesh.Attributes()->GetMaterialID();
+					FDynamicMeshMaterialAttribute* ToMaterialIDAttrib = InMergedMesh.Attributes()->GetMaterialID();
+					TMap<int32, int32> MaterialMap;
+					for (const TPair<int32, int32>& FromToTId : TmpMappings.GetTriangleMap().GetForwardMap())
+					{
+						const int32 FromMatId = FromMaterialIDAttrib->GetValue(FromToTId.Key);
+						const int32 ToMatId = FromMatId + MaterialCount;
+						MaterialMap.Add(FromMatId, ToMatId);
+						ToMaterialIDAttrib->SetNewValue(FromToTId.Value, ToMatId);
+					}
+					MaterialCount += MaterialMap.Num();
+
+					for (TWeakObjectPtr<UMaterialInterface>& BakedMaterial : AttachmentItem->BakedMaterials)
+					{
+						UMaterialInterface* MaterialInterface = BakedMaterial.Get();
+						int32 MatIdx = CombineMaterials.Add(MaterialInterface);
+#if WITH_EDITOR
+						if (MaterialInterface)
+						{
+							const UMaterial* Material = MaterialInterface->GetMaterial();
+							if (Material && !Material->GetUsageByFlag(EMaterialUsage::MATUSAGE_NiagaraMeshParticles))
+							{
+								MaterialsMissingNiagaraUsageFlag.Add(MatIdx);
+							}
+						}
+#endif
+					}
+				}
+			});
+		}
+	}, EDynamicMeshChangeType::GeneralEdit, EDynamicMeshAttributeChangeFlags::Unknown, true);
+
+#if WITH_EDITOR
+	if (!MaterialsMissingNiagaraUsageFlag.IsEmpty())
+	{
+		int32 ReadOnlyMaterialCount = 0;
+
+		for (const int32 MaterialIdx : MaterialsMissingNiagaraUsageFlag)
+		{
+			UMaterialInterface* Material = CombineMaterials[MaterialIdx];
+			if (!Material)
+			{
+				continue;
+			}
+
+			if (!IsMaterialDirtyable(Material))
+			{
+				ReadOnlyMaterialCount++;
+				CombineMaterials[MaterialIdx] = nullptr;
+				UE_LOG(LogCEClonerComponent, Warning, TEXT("%s : The following materials (%s) on actor (%s) does not have the required usage flag (bUsedWithNiagaraMeshParticles) to work with the cloner, this material cannot be dirtied due to its read-only location, skipping material and proceeding"), *ClonerActor->GetActorNameOrLabel(), Material ? *Material->GetMaterial()->GetPathName() : TEXT("Invalid Material"), *InRootActor->GetActorNameOrLabel());
+			}
+			else
+			{
+				UE_LOG(LogCEClonerComponent, Log, TEXT("%s : The following materials (%s) on actor (%s) does not have the required usage flag (bUsedWithNiagaraMeshParticles) to work with the cloner, dirtying material and proceeding, please save the asset for working runtime result"), *ClonerActor->GetActorNameOrLabel(), Material ? *Material->GetMaterial()->GetPathName() : TEXT("Invalid Material"), *InRootActor->GetActorNameOrLabel());
+			}
+		}
+
+		if (ReadOnlyMaterialCount > 0)
+		{
+			FNotificationInfo NotificationInfo(FText::Format(LOCTEXT("MaterialsMissingUsageFlag", "Detected {0} read-only material(s) with missing niagara usage flag required to work properly with cloner (See logs)"), ReadOnlyMaterialCount));
+			NotificationInfo.ExpireDuration = 5.f;
+			NotificationInfo.bFireAndForget = true;
+			NotificationInfo.Image = FAppStyle::GetBrush("Icons.WarningWithColor");
+			FSlateNotificationManager::Get().AddNotification(NotificationInfo);
+		}
+	}
+#endif
+
+	// Convert merged mesh to static mesh
+	UStaticMesh* BakedAttachmentMesh = NewObject<UStaticMesh>();
+	if (MergedAttachmentMesh->GetTriangleCount())
+	{
+		// export options
+		FGeometryScriptCopyMeshToAssetOptions AssetOptions;
+		AssetOptions.bReplaceMaterials = true;
+		AssetOptions.bEnableRecomputeNormals = false;
+		AssetOptions.bEnableRecomputeTangents = false;
+		AssetOptions.bEnableRemoveDegenerates = true;
+		AssetOptions.bEmitTransaction = false;
+		AssetOptions.bApplyNaniteSettings = false;
+		AssetOptions.bDeferMeshPostEditChange = false;
+		AssetOptions.NewMaterials = CombineMaterials;
+		// LOD options
+		static FGeometryScriptMeshWriteLOD TargetLOD;
+		TargetLOD.LODIndex = 0;
+		TargetLOD.bWriteHiResSource = false;
+		// result
+		EGeometryScriptOutcomePins OutResult;
+		UGeometryScriptLibrary_StaticMeshFunctions::CopyMeshToStaticMesh(MergedAttachmentMesh, BakedAttachmentMesh, AssetOptions, TargetLOD, OutResult);
+
+		// Fix camera culling and avoid killing particles when bounds partially in view
+		const float NewRadiusBounds = BakedAttachmentMesh->GetBounds().SphereRadius * 2;
+		BakedAttachmentMesh->SetExtendedBounds(FBox(-FVector(NewRadiusBounds), FVector(NewRadiusBounds)));
+
+#if WITH_EDITORONLY_DATA
+		// Compute normals and tangents
+		FMeshDescription* MeshDescription = BakedAttachmentMesh->GetMeshDescription(0);
+		FStaticMeshOperations::ComputeTriangleTangentsAndNormals(*MeshDescription);
+		FStaticMeshOperations::ComputeTangentsAndNormals(*MeshDescription, EComputeNTBsFlags::Normals);
+#endif
+
+		bClonerMeshesDirty = true;
+	}
+
+	MergedAttachmentMesh->MarkAsGarbage();
+	ClonerTree.MergedBakedMeshes[RootIdx] = BakedAttachmentMesh;
+}
+
+void UCEClonerComponent::GetActorAttachmentItems(AActor* InActor, TArray<FCEClonerAttachmentItem*>& OutAttachmentItems)
+{
+	if (!InActor)
+	{
+		return;
+	}
+
+	FCEClonerAttachmentItem* AttachmentItem = ClonerTree.ItemAttachmentMap.Find(InActor);
+
+	if (!AttachmentItem)
+	{
+		return;
+	}
+
+	OutAttachmentItems.Add(AttachmentItem);
+
+	for (TWeakObjectPtr<AActor>& ChildActor : AttachmentItem->ChildrenActors)
+	{
+		if (ChildActor.IsValid())
+		{
+			GetActorAttachmentItems(ChildActor.Get(), OutAttachmentItems);
+		}
+	}
+}
+
+bool UCEClonerComponent::IsAllMergedMeshesValid() const
+{
+	for (const TObjectPtr<UStaticMesh>& MergedMesh : ClonerTree.MergedBakedMeshes)
+	{
+		if (!MergedMesh)
+		{
+			return false;
+		}
+	}
+	return true;
+}
+
+bool UCEClonerComponent::IsMaterialDirtyable(const UMaterialInterface* InMaterial) const
+{
+	const UMaterial* BaseMaterial = InMaterial->GetMaterial();
+	const FString ContentFolder = FPaths::ConvertRelativePathToFull(FPaths::ProjectContentDir());
+
+	const UPackage* MaterialPackage = BaseMaterial->GetPackage();
+	const FPackagePath& LoadedPath = MaterialPackage->GetLoadedPath();
+	const FString PackagePath = FPaths::ConvertRelativePathToFull(LoadedPath.GetLocalFullPath());
+	const FString MaterialPath = BaseMaterial->GetPathName();
+
+	const bool bTransientPackage = MaterialPackage == GetTransientPackage() || MaterialPath.StartsWith("/Temp/");
+	const bool bContentFolder = PackagePath.StartsWith(ContentFolder);
+
+	return bTransientPackage || bContentFolder;
+}
+
+void UCEClonerComponent::UpdateClonerMeshes()
+{
+	const AActor* ClonerActor = GetOwner();
+	if (!ClonerActor)
+	{
+		return;
+	}
+
+	if (!bClonerMeshesDirty)
+	{
+		return;
+	}
+
+	if (!GetAsset() || !ActiveLayout)
+	{
+		return;
+	}
+
+	UNiagaraMeshRendererProperties* MeshRenderer = ActiveLayout->GetMeshRenderer();
+
+	if (!MeshRenderer)
+	{
+		UE_LOG(LogCEClonerComponent, Warning, TEXT("%s : Invalid mesh renderer for cloner system"), *ClonerActor->GetActorNameOrLabel());
+		return;
+	}
+
+	// Set new number of meshes in renderer
+	static const FName MeshNumName(TEXT("MeshNum"));
+	SetIntParameter(MeshNumName, ClonerTree.MergedBakedMeshes.Num());
+
+	// Resize mesh array properly
+	if (MeshRenderer->Meshes.Num() > ClonerTree.MergedBakedMeshes.Num())
+	{
+		MeshRenderer->Meshes.SetNum(ClonerTree.MergedBakedMeshes.Num());
+	}
+
+	// Set baked meshes in mesh renderer array
+	for (int32 Idx = 0; Idx < ClonerTree.MergedBakedMeshes.Num(); Idx++)
+	{
+		UStaticMesh* StaticMesh = ClonerTree.MergedBakedMeshes[Idx];
+		FNiagaraMeshRendererMeshProperties& MeshProperties = !MeshRenderer->Meshes.IsValidIndex(Idx) ? MeshRenderer->Meshes.AddDefaulted_GetRef() : MeshRenderer->Meshes[Idx];
+		MeshProperties.Mesh = StaticMesh && StaticMesh->GetNumTriangles(0) > 0 ? StaticMesh : nullptr;
+
+		if (ClonerTree.RootActors.IsValidIndex(Idx))
+		{
+			if (const FCEClonerAttachmentItem* Item = ClonerTree.ItemAttachmentMap.Find(ClonerTree.RootActors[Idx]))
+			{
+				MeshProperties.Rotation = Item->ActorTransform.Rotator();
+				MeshProperties.Scale = Item->ActorTransform.GetScale3D();
+			}
+		}
+	}
+
+#if WITH_EDITORONLY_DATA
+	MeshRenderer->OnMeshChanged();
+#endif
+
+	bClonerMeshesDirty = !DirtyItemAttachments.IsEmpty();
+
+	UE_LOG(LogCEClonerComponent, Log, TEXT("%s : Cloner mesh updated %i"), *ClonerActor->GetActorNameOrLabel(), ClonerTree.MergedBakedMeshes.Num());
+
+	OnClonerMeshUpdated.Broadcast();
+}
+
+bool UCEClonerComponent::SetClonerActiveLayout(UCEClonerLayoutBase* InLayout)
+{
+	if (!InLayout)
+	{
+		return false;
+	}
+
+	const AActor* ClonerActor = GetOwner();
+	if (!ClonerActor)
+	{
+		return false;
+	}
+
+	if (!InLayout->IsLayoutLoaded())
+	{
+		// Load system
+		if (!InLayout->LoadLayout())
+		{
+			UE_LOG(LogCEClonerComponent, Warning, TEXT("%s : Cloner layout system failed to load %s - %s"), *GetOwner()->GetActorNameOrLabel(), *InLayout->GetLayoutName().ToString(), *InLayout->GetLayoutAssetPath());
+			return false;
+		}
+
+		UE_LOG(LogCEClonerComponent, Log, TEXT("%s : Cloner layout system loaded %s - %s"), *GetOwner()->GetActorNameOrLabel(), *InLayout->GetLayoutName().ToString(), *InLayout->GetLayoutAssetPath());
+
+		OnClonerSystemLoaded.Broadcast(InLayout);
+	}
+
+	// Copy data interfaces to new layout
+	if (ActiveLayout && ActiveLayout->IsLayoutLoaded())
+	{
+		ActiveLayout->CopyTo(InLayout);
+	}
+
+	// Deactivate previous layout
+	if (ActiveLayout && ActiveLayout->IsLayoutActive())
+	{
+		ActiveLayout->DeactivateLayout();
+	}
+
+	// Activate new layout
+	InLayout->ActivateLayout();
+
+	ActiveLayout = InLayout;
+	bClonerMeshesDirty = true;
+
+	UE_LOG(LogCEClonerComponent, Log, TEXT("%s : Cloner layout system changed %s - %s"), *ClonerActor->GetActorNameOrLabel(), *InLayout->GetLayoutName().ToString(), *InLayout->GetLayoutAssetPath());
+
+	return true;
+}
+
+#undef LOCTEXT_NAMESPACE
diff --git a/Engine/Plugins/Experimental/ClonerEffector/Source/ClonerEffector/Private/Cloner/Layouts/CEClonerGridLayout.cpp b/Engine/Plugins/Experimental/ClonerEffector/Source/ClonerEffector/Private/Cloner/Layouts/CEClonerGridLayout.cpp
new file mode 100644
--- /dev/null
+++ b/Engine/Plugins/Experimental/ClonerEffector/Source/ClonerEffector/Private/Cloner/Layouts/CEClonerGridLayout.cpp
@@ -0,0 +1,181 @@
+// Copyright Epic Games, Inc. All Rights Reserved.
+
+#include "Cloner/Layouts/CEClonerGridLayout.h"
+
+#include "Cloner/CEClonerComponent.h"
+#include "NiagaraDataInterfaceTexture.h"
+#include "NiagaraSystem.h"
+
+void UCEClonerGridLayout::SetCountX(int32 InCountX)
+{
+	if (CountX == InCountX)
+	{
+		return;
+	}
+
+	CountX = InCountX;
+	UpdateLayoutParameters();
+}
+
+void UCEClonerGridLayout::SetCountY(int32 InCountY)
+{
+	if (CountY == InCountY)
+	{
+		return;
+	}
+
+	CountY = InCountY;
+	UpdateLayoutParameters();
+}
+
+void UCEClonerGridLayout::SetCountZ(int32 InCountZ)
+{
+	if (CountZ == InCountZ)
+	{
+		return;
+	}
+
+	CountZ = InCountZ;
+	UpdateLayoutParameters();
+}
+
+void UCEClonerGridLayout::SetSpacingX(float InSpacingX)
+{
+	if (SpacingX == InSpacingX)
+	{
+		return;
+	}
+
+	SpacingX = InSpacingX;
+	UpdateLayoutParameters();
+}
+
+void UCEClonerGridLayout::SetSpacingY(float InSpacingY)
+{
+	if (SpacingY == InSpacingY)
+	{
+		return;
+	}
+
+	SpacingY = InSpacingY;
+	UpdateLayoutParameters();
+}
+
+void UCEClonerGridLayout::SetSpacingZ(float InSpacingZ)
+{
+	if (SpacingZ == InSpacingZ)
+	{
+		return;
+	}
+
+	SpacingZ = InSpacingZ;
+	UpdateLayoutParameters();
+}
+
+void UCEClonerGridLayout::SetConstraint(ECEClonerGridConstraint InConstraint)
+{
+	if (Constraint == InConstraint)
+	{
+		return;
+	}
+
+	Constraint = InConstraint;
+	UpdateLayoutParameters();
+}
+
+void UCEClonerGridLayout::SetInvertConstraint(bool bInInvertConstraint)
+{
+	if (bInvertConstraint == bInInvertConstraint)
+	{
+		return;
+	}
+
+	bInvertConstraint = bInInvertConstraint;
+	UpdateLayoutParameters();
+}
+
+void UCEClonerGridLayout::SetSphereConstraint(const FCEClonerGridConstraintSphere& InConstraint)
+{
+	SphereConstraint = InConstraint;
+	UpdateLayoutParameters();
+}
+
+void UCEClonerGridLayout::SetCylinderConstraint(const FCEClonerGridConstraintCylinder& InConstraint)
+{
+	CylinderConstraint = InConstraint;
+	UpdateLayoutParameters();
+}
+
+void UCEClonerGridLayout::SetTextureConstraint(const FCEClonerGridConstraintTexture& InConstraint)
+{
+	TextureConstraint = InConstraint;
+	UpdateLayoutParameters();
+}
+
+#if WITH_EDITOR
+const TCEPropertyChangeDispatcher<UCEClonerGridLayout> UCEClonerGridLayout::PropertyChangeDispatcher =
+{
+	{ GET_MEMBER_NAME_CHECKED(UCEClonerGridLayout, CountX), &UCEClonerGridLayout::OnLayoutPropertyChanged },
+	{ GET_MEMBER_NAME_CHECKED(UCEClonerGridLayout, CountY), &UCEClonerGridLayout::OnLayoutPropertyChanged },
+	{ GET_MEMBER_NAME_CHECKED(UCEClonerGridLayout, CountZ), &UCEClonerGridLayout::OnLayoutPropertyChanged },
+	{ GET_MEMBER_NAME_CHECKED(UCEClonerGridLayout, SpacingX), &UCEClonerGridLayout::OnLayoutPropertyChanged },
+	{ GET_MEMBER_NAME_CHECKED(UCEClonerGridLayout, SpacingY), &UCEClonerGridLayout::OnLayoutPropertyChanged },
+	{ GET_MEMBER_NAME_CHECKED(UCEClonerGridLayout, SpacingZ), &UCEClonerGridLayout::OnLayoutPropertyChanged },
+	{ GET_MEMBER_NAME_CHECKED(UCEClonerGridLayout, Constraint), &UCEClonerGridLayout::OnLayoutPropertyChanged },
+	{ GET_MEMBER_NAME_CHECKED(UCEClonerGridLayout, bInvertConstraint), &UCEClonerGridLayout::OnLayoutPropertyChanged },
+	{ GET_MEMBER_NAME_CHECKED(UCEClonerGridLayout, SphereConstraint), &UCEClonerGridLayout::OnLayoutPropertyChanged },
+	{ GET_MEMBER_NAME_CHECKED(UCEClonerGridLayout, CylinderConstraint), &UCEClonerGridLayout::OnLayoutPropertyChanged },
+	{ GET_MEMBER_NAME_CHECKED(UCEClonerGridLayout, TextureConstraint), &UCEClonerGridLayout::OnLayoutPropertyChanged },
+};
+
+void UCEClonerGridLayout::PostEditChangeProperty(FPropertyChangedEvent& InPropertyChangedEvent)
+{
+	Super::PostEditChangeProperty(InPropertyChangedEvent);
+
+	PropertyChangeDispatcher.OnPropertyChanged(this, InPropertyChangedEvent);
+}
+#endif
+
+void UCEClonerGridLayout::OnLayoutParametersChanged(UCEClonerComponent* InComponent)
+{
+	Super::OnLayoutParametersChanged(InComponent);
+
+	InComponent->SetIntParameter(TEXT("GridCountX"), CountX);
+
+	InComponent->SetIntParameter(TEXT("GridCountY"), CountY);
+
+	InComponent->SetIntParameter(TEXT("GridCountZ"), CountZ);
+
+	InComponent->SetVectorParameter(TEXT("GridSpacing"), FVector(SpacingX, SpacingY, SpacingZ));
+
+	FNiagaraUserRedirectionParameterStore& ExposedParameters = InComponent->GetAsset()->GetExposedParameters();
+	static const FNiagaraVariable ConstraintVar(FNiagaraTypeDefinition(StaticEnum<ECEClonerGridConstraint>()), TEXT("Constraint"));
+	ExposedParameters.SetParameterValue<int32>(static_cast<int32>(Constraint), ConstraintVar);
+
+	InComponent->SetBoolParameter(TEXT("ConstraintInvert"), Constraint != ECEClonerGridConstraint::None ? bInvertConstraint : false);
+
+	InComponent->SetVectorParameter(TEXT("ConstraintCylinderCenter"), CylinderConstraint.Center);
+
+	InComponent->SetFloatParameter(TEXT("ConstraintCylinderHeight"), CylinderConstraint.Height);
+
+	InComponent->SetFloatParameter(TEXT("ConstraintCylinderRadius"), CylinderConstraint.Radius);
+
+	InComponent->SetVectorParameter(TEXT("ConstraintSphereCenter"), SphereConstraint.Center);
+
+	InComponent->SetFloatParameter(TEXT("ConstraintSphereRadius"), SphereConstraint.Radius);
+
+	static const FNiagaraVariable ConstraintTextureChannelVar(FNiagaraTypeDefinition(StaticEnum<ECEClonerTextureSampleChannel>()), TEXT("ConstraintTextureChannel"));
+	ExposedParameters.SetParameterValue<int32>(static_cast<int32>(TextureConstraint.Channel), ConstraintTextureChannelVar);
+
+	static const FNiagaraVariable ConstraintTextureCompareModeVar(FNiagaraTypeDefinition(StaticEnum<ECEClonerCompareMode>()), TEXT("ConstraintTextureCompareMode"));
+	ExposedParameters.SetParameterValue<int32>(static_cast<int32>(TextureConstraint.CompareMode), ConstraintTextureCompareModeVar);
+
+	static const FNiagaraVariable ConstraintTexturePlaneVar(FNiagaraTypeDefinition(StaticEnum<ECEClonerPlane>()), TEXT("ConstraintTexturePlane"));
+	ExposedParameters.SetParameterValue<int32>(static_cast<int32>(TextureConstraint.Plane), ConstraintTexturePlaneVar);
+
+	InComponent->SetFloatParameter(TEXT("ConstraintTextureThreshold"), TextureConstraint.Threshold);
+
+	static const FNiagaraVariable ConstraintTextureSamplerVar(FNiagaraTypeDefinition(UNiagaraDataInterfaceTexture::StaticClass()), TEXT("ConstraintTextureSampler"));
+	UNiagaraDataInterfaceTexture* TextureSamplerDI = Cast<UNiagaraDataInterfaceTexture>(ExposedParameters.GetDataInterface(ConstraintTextureSamplerVar));
+	TextureSamplerDI->SetTexture(TextureConstraint.Texture.Get());
+}
