MotionDesign

diff --git a/Engine/Plugins/Experimental/ClonerEffector/Source/ClonerEffector/Private/Cloner/Layouts/CEClonerLayoutBase.cpp b/Engine/Plugins/Experimental/ClonerEffector/Source/ClonerEffector/Private/Cloner/Layouts/CEClonerLayoutBase.cpp
new file mode 100644
--- /dev/null
+++ b/Engine/Plugins/Experimental/ClonerEffector/Source/ClonerEffector/Private/Cloner/Layouts/CEClonerLayoutBase.cpp
@@ -0,0 +1,289 @@
+// Copyright Epic Games, Inc. All Rights Reserved.
+
+#include "Cloner/Layouts/CEClonerLayoutBase.h"
+
+#include "Cloner/CEClonerActor.h"
+#include "Cloner/CEClonerComponent.h"
+#include "NiagaraEmitter.h"
+#include "NiagaraMeshRendererProperties.h"
+#include "NiagaraSystem.h"
+#include "UObject/Package.h"
+#include "UObject/UObjectGlobals.h"
+
+DEFINE_LOG_CATEGORY_STATIC(LogCEClonerLayoutBase, Log, All);
+
+bool UCEClonerLayoutBase::IsLayoutValid() const
+{
+	if (LayoutName.IsNone() || LayoutAssetPath.IsEmpty())
+	{
+		return false;
+	}
+
+	// Get the template niagara asset
+	const UNiagaraSystem* TemplateNiagaraSystem = LoadSystemPath(LayoutAssetPath);
+
+	// Get the base niagara asset
+	const UNiagaraSystem* BaseNiagaraSystem = LoadSystemPath(LayoutBaseAssetPath);
+
+	if (!TemplateNiagaraSystem || !BaseNiagaraSystem)
+	{
+		UE_LOG(LogCEClonerLayoutBase, Warning, TEXT("Cloner layout %s : Template system (%s) or base system (%s) is invalid"), *LayoutName.ToString(), *LayoutAssetPath, LayoutBaseAssetPath);
+		return false;
+	}
+
+	// Compare parameters : template should have base parameters
+	bool bIsSystemBasedOnBaseAsset = true;
+
+	{
+		TArray<FNiagaraVariable> TemplateSystemParameters;
+		TemplateNiagaraSystem->GetExposedParameters().GetParameters(TemplateSystemParameters);
+
+		TArray<FNiagaraVariable> BaseSystemParameters;
+		BaseNiagaraSystem->GetExposedParameters().GetParameters(BaseSystemParameters);
+
+		for (const FNiagaraVariable& SystemParameter : BaseSystemParameters)
+		{
+			if (!TemplateSystemParameters.Contains(SystemParameter))
+			{
+				bIsSystemBasedOnBaseAsset = false;
+				UE_LOG(LogCEClonerLayoutBase, Warning, TEXT("Cloner layout %s : Template system (%s) missing parameter (%s) from base system (%s)"), *LayoutName.ToString(), *LayoutAssetPath, *SystemParameter.ToString(), LayoutBaseAssetPath);
+				break;
+			}
+		}
+	}
+
+	if (!bIsSystemBasedOnBaseAsset)
+	{
+		UE_LOG(LogCEClonerLayoutBase, Warning, TEXT("Cloner layout %s : Template system (%s) is not based off base system (%s)"), *LayoutName.ToString(), *LayoutAssetPath, LayoutBaseAssetPath);
+	}
+
+	return bIsSystemBasedOnBaseAsset;
+}
+
+bool UCEClonerLayoutBase::IsLayoutLoaded() const
+{
+	return !IsTemplate() && NiagaraSystem && MeshRenderer && DataInterfaces.IsValid();
+}
+
+bool UCEClonerLayoutBase::LoadLayout()
+{
+	if (IsLayoutLoaded())
+	{
+		return true;
+	}
+
+	if (LayoutAssetPath.IsEmpty())
+	{
+		return false;
+	}
+
+	// Get the template niagara asset
+	const UNiagaraSystem* TemplateNiagaraSystem = FindObject<UNiagaraSystem>(nullptr, *LayoutAssetPath);
+	if (!TemplateNiagaraSystem)
+	{
+		TemplateNiagaraSystem = LoadObject<UNiagaraSystem>(nullptr, *LayoutAssetPath);
+	}
+
+	if (!TemplateNiagaraSystem)
+	{
+		return false;
+	}
+
+	// Copy template asset since we will modify it, and we do not want other cloner instances to have the new modified asset
+	FObjectDuplicationParameters Parameters = InitStaticDuplicateObjectParams(TemplateNiagaraSystem, GetTransientPackage(), NAME_None, RF_Transient, nullptr, EDuplicateMode::Normal);
+	NiagaraSystem = Cast<UNiagaraSystem>(StaticDuplicateObjectEx(Parameters));
+
+	if (!NiagaraSystem)
+	{
+		return false;
+	}
+
+	for (FNiagaraEmitterHandle& SystemEmitterHandle : NiagaraSystem->GetEmitterHandles())
+	{
+		if (const FVersionedNiagaraEmitterData* EmitterData = SystemEmitterHandle.GetEmitterData())
+		{
+			for (UNiagaraRendererProperties* EmitterRenderer : EmitterData->GetRenderers())
+			{
+				if (UNiagaraMeshRendererProperties* EmitterMeshRenderer = Cast<UNiagaraMeshRendererProperties>(EmitterRenderer))
+				{
+					EmitterMeshRenderer->Meshes.Empty();
+#if WITH_EDITORONLY_DATA
+					EmitterMeshRenderer->OnMeshChanged();
+#endif
+
+					MeshRenderer = EmitterMeshRenderer;
+					DataInterfaces = FCEClonerEffectorDataInterfaces(NiagaraSystem);
+
+					OnLayoutLoaded();
+
+					return true;
+				}
+			}
+		}
+	}
+
+	return false;
+}
+
+bool UCEClonerLayoutBase::UnloadLayout()
+{
+	if (!IsLayoutLoaded())
+	{
+		return false;
+	}
+
+	// Cannot unload while active
+	if (IsLayoutActive())
+	{
+		return false;
+	}
+
+	DataInterfaces = FCEClonerEffectorDataInterfaces();
+	MeshRenderer = nullptr;
+	NiagaraSystem = nullptr;
+
+	OnLayoutUnloaded();
+
+	return true;
+}
+
+bool UCEClonerLayoutBase::IsLayoutActive() const
+{
+	const UCEClonerComponent* Component = GetClonerComponent();
+
+	if (!Component)
+	{
+		return false;
+	}
+
+	return IsLayoutLoaded() && Component->GetAsset() == NiagaraSystem;
+}
+
+bool UCEClonerLayoutBase::ActivateLayout()
+{
+	if (IsLayoutActive())
+	{
+		return false;
+	}
+
+	// Load layout first
+	if (!IsLayoutLoaded())
+	{
+		return false;
+	}
+
+	UCEClonerComponent* ClonerComponent = GetClonerComponent();
+
+	if (!ClonerComponent)
+	{
+		return false;
+	}
+
+	ClonerComponent->SetAsset(NiagaraSystem);
+
+	OnLayoutActive();
+
+	return true;
+}
+
+bool UCEClonerLayoutBase::DeactivateLayout()
+{
+	if (!IsLayoutActive())
+	{
+		return false;
+	}
+
+	UCEClonerComponent* ClonerComponent = GetClonerComponent();
+
+	if (!ClonerComponent)
+	{
+		return false;
+	}
+
+	ClonerComponent->SetAsset(nullptr);
+
+	OnLayoutInactive();
+
+	return true;
+}
+
+bool UCEClonerLayoutBase::CopyTo(UCEClonerLayoutBase* InOtherLayout) const
+{
+	if (!InOtherLayout || InOtherLayout == this)
+	{
+		return false;
+	}
+
+	if (!IsLayoutLoaded())
+	{
+		return false;
+	}
+
+	const UNiagaraSystem* OtherSystem = InOtherLayout->NiagaraSystem;
+	if (!OtherSystem)
+	{
+		return false;
+	}
+
+	DataInterfaces.CopyTo(InOtherLayout->DataInterfaces);
+
+	return true;
+}
+
+void UCEClonerLayoutBase::OnLayoutPropertyChanged()
+{
+	UpdateLayoutParameters();
+}
+
+UNiagaraSystem* UCEClonerLayoutBase::LoadSystemPath(const FString& InPath) const
+{
+	UNiagaraSystem* LoadedNiagaraSystem = FindObject<UNiagaraSystem>(nullptr, *InPath);
+
+	if (!LoadedNiagaraSystem)
+	{
+		LoadedNiagaraSystem = LoadObject<UNiagaraSystem>(nullptr, *InPath);
+	}
+
+	return LoadedNiagaraSystem;
+}
+
+ACEClonerActor* UCEClonerLayoutBase::GetClonerActor() const
+{
+	return GetTypedOuter<ACEClonerActor>();
+}
+
+UCEClonerComponent* UCEClonerLayoutBase::GetClonerComponent() const
+{
+	if (const ACEClonerActor* Cloner = GetClonerActor())
+	{
+		return Cloner->GetClonerComponent();
+	}
+
+	return nullptr;
+}
+
+void UCEClonerLayoutBase::UpdateLayoutParameters(bool bInUpdateCloner, bool bInImmediate)
+{
+	if (!IsLayoutActive())
+	{
+		return;
+	}
+
+	if (UCEClonerComponent* ClonerComponent = GetClonerComponent())
+	{
+		OnLayoutParametersChanged(ClonerComponent);
+
+		if (bInUpdateCloner)
+		{
+			RequestClonerUpdate(bInImmediate);
+		}
+	}
+}
+
+void UCEClonerLayoutBase::RequestClonerUpdate(bool bInImmediate) const
+{
+	if (ACEClonerActor* Cloner = GetClonerActor())
+	{
+		Cloner->RequestClonerUpdate(bInImmediate);
+	}
+}
diff --git a/Engine/Plugins/Experimental/ClonerEffector/Source/ClonerEffector/Private/Effector/CEEffectorActor.cpp b/Engine/Plugins/Experimental/ClonerEffector/Source/ClonerEffector/Private/Effector/CEEffectorActor.cpp
new file mode 100644
--- /dev/null
+++ b/Engine/Plugins/Experimental/ClonerEffector/Source/ClonerEffector/Private/Effector/CEEffectorActor.cpp
@@ -0,0 +1,1255 @@
+// Copyright Epic Games, Inc. All Rights Reserved.
+
+#include "Effector/CEEffectorActor.h"
+
+#include "Cloner/CEClonerActor.h"
+#include "Components/BoxComponent.h"
+#include "Components/SphereComponent.h"
+#include "Effector/CEEffectorComponent.h"
+#include "Math/Vector.h"
+#include "Subsystems/CEEffectorSubsystem.h"
+
+DEFINE_LOG_CATEGORY_STATIC(LogCEEffectorActor, Log, All);
+
+ACEEffectorActor::ACEEffectorActor()
+{
+	SetCanBeDamaged(false);
+	PrimaryActorTick.bCanEverTick = false;
+
+	USceneComponent* SceneComponent = CreateDefaultSubobject<UCEEffectorComponent>(TEXT("AvaEffectorComponent"));
+#if WITH_EDITORONLY_DATA
+	SceneComponent->bVisualizeComponent = true;
+#endif
+	SetRootComponent(SceneComponent);
+
+	// Sphere
+	InnerSphereComponent = CreateDefaultSubobject<USphereComponent>(TEXT("AvaInnerSphereComponent"));;
+	InnerSphereComponent->ShapeColor = FColor::Blue;
+	InnerSphereComponent->SetLineThickness(VisualizerThickness);
+	InnerSphereComponent->SetSphereRadius(InnerRadius);
+	InnerSphereComponent->SetHiddenInGame(true);
+#if WITH_EDITOR
+	// Do not show bounding box around cloner for better visibility
+	InnerSphereComponent->SetIsVisualizationComponent(true);
+#endif
+	InnerSphereComponent->bIsEditorOnly = false;
+	InnerSphereComponent->SetupAttachment(SceneComponent);
+
+	OuterSphereComponent = CreateDefaultSubobject<USphereComponent>(TEXT("AvaOuterSphereComponent"));
+	OuterSphereComponent->ShapeColor = FColor::Red;
+	OuterSphereComponent->SetLineThickness(VisualizerThickness);
+	OuterSphereComponent->SetSphereRadius(OuterRadius);
+	OuterSphereComponent->SetHiddenInGame(true);
+#if WITH_EDITOR
+	// Do not show bounding box around cloner for better visibility
+	OuterSphereComponent->SetIsVisualizationComponent(true);
+#endif
+	OuterSphereComponent->bIsEditorOnly = false;
+	OuterSphereComponent->SetupAttachment(SceneComponent);
+
+	// Box
+	InnerBoxComponent = CreateDefaultSubobject<UBoxComponent>(TEXT("AvaInnerBoxComponent"));
+	InnerBoxComponent->ShapeColor = FColor::Blue;
+	InnerBoxComponent->SetLineThickness(VisualizerThickness);
+	InnerBoxComponent->SetBoxExtent(InnerExtent);
+	InnerBoxComponent->SetHiddenInGame(true);
+#if WITH_EDITOR
+	// Do not show bounding box around cloner for better visibility
+	InnerBoxComponent->SetIsVisualizationComponent(true);
+#endif
+	InnerBoxComponent->bIsEditorOnly = false;
+	InnerBoxComponent->SetupAttachment(SceneComponent);
+
+	OuterBoxComponent = CreateDefaultSubobject<UBoxComponent>(TEXT("AvaOuterBoxComponent"));
+	OuterBoxComponent->ShapeColor = FColor::Red;
+	OuterBoxComponent->SetLineThickness(VisualizerThickness);
+	OuterBoxComponent->SetBoxExtent(OuterExtent);
+	OuterBoxComponent->SetHiddenInGame(true);
+#if WITH_EDITOR
+	// Do not show bounding box around cloner for better visibility
+	OuterBoxComponent->SetIsVisualizationComponent(true);
+#endif
+	OuterBoxComponent->bIsEditorOnly = false;
+	OuterBoxComponent->SetupAttachment(SceneComponent);
+
+	// Plane
+	InnerPlaneComponent = CreateDefaultSubobject<UBoxComponent>(TEXT("AvaInnerPlaneComponent"));
+	InnerPlaneComponent->ShapeColor = FColor::Blue;
+	InnerPlaneComponent->SetLineThickness(VisualizerThickness);
+	InnerPlaneComponent->SetBoxExtent(InnerExtent);
+	InnerPlaneComponent->SetHiddenInGame(true);
+#if WITH_EDITOR
+	// Do not show bounding box around cloner for better visibility
+	InnerPlaneComponent->SetIsVisualizationComponent(true);
+#endif
+	InnerPlaneComponent->bIsEditorOnly = false;
+	InnerPlaneComponent->SetupAttachment(SceneComponent);
+
+	OuterPlaneComponent = CreateDefaultSubobject<UBoxComponent>(TEXT("AvaOuterPlaneComponent"));
+	OuterPlaneComponent->ShapeColor = FColor::Red;
+	OuterPlaneComponent->SetLineThickness(VisualizerThickness);
+	OuterPlaneComponent->SetBoxExtent(OuterExtent);
+	OuterPlaneComponent->SetHiddenInGame(true);
+#if WITH_EDITOR
+	// Do not show bounding box around cloner for better visibility
+	OuterPlaneComponent->SetIsVisualizationComponent(true);
+#endif
+	OuterPlaneComponent->bIsEditorOnly = false;
+	OuterPlaneComponent->SetupAttachment(SceneComponent);
+
+	if (!IsTemplate())
+	{
+		SceneComponent->TransformUpdated.AddUObject(this, &ACEEffectorActor::OnEffectorTransformed);
+
+		RegisterToChannel();
+	}
+}
+
+#if WITH_EDITOR
+FString ACEEffectorActor::GetDefaultActorLabel() const
+{
+	return DefaultLabel;
+}
+
+TCEPropertyChangeDispatcher<ACEEffectorActor> ACEEffectorActor::PropertyChangeDispatcher =
+{
+	/** Effector */
+	{ GET_MEMBER_NAME_CHECKED(ACEEffectorActor, bEnabled), &ACEEffectorActor::OnEnabledChanged },
+	{ GET_MEMBER_NAME_CHECKED(ACEEffectorActor, Magnitude), &ACEEffectorActor::OnMagnitudeChanged },
+	{ GET_MEMBER_NAME_CHECKED(ACEEffectorActor, Cloners), &ACEEffectorActor::OnClonersChanged },
+	{ GET_MEMBER_NAME_CHECKED(ACEEffectorActor, VisualizerThickness), &ACEEffectorActor::OnVisualizerThicknessChanged },
+	{ GET_MEMBER_NAME_CHECKED(ACEEffectorActor, bVisualizerSpriteVisible), &ACEEffectorActor::OnVisualizerSpriteVisibleChanged },
+	/** Type */
+	{ GET_MEMBER_NAME_CHECKED(ACEEffectorActor, Type), &ACEEffectorActor::OnTypeChanged },
+	{ GET_MEMBER_NAME_CHECKED(ACEEffectorActor, Easing), &ACEEffectorActor::OnEasingChanged },
+	{ GET_MEMBER_NAME_CHECKED(ACEEffectorActor, OuterRadius), &ACEEffectorActor::OnSphereChanged },
+	{ GET_MEMBER_NAME_CHECKED(ACEEffectorActor, InnerRadius), &ACEEffectorActor::OnSphereChanged },
+	{ GET_MEMBER_NAME_CHECKED(ACEEffectorActor, InnerExtent), &ACEEffectorActor::OnBoxChanged },
+	{ GET_MEMBER_NAME_CHECKED(ACEEffectorActor, OuterExtent), &ACEEffectorActor::OnBoxChanged },
+	{ GET_MEMBER_NAME_CHECKED(ACEEffectorActor, PlaneSpacing), &ACEEffectorActor::OnPlaneChanged },
+	/** Mode */
+	{ GET_MEMBER_NAME_CHECKED(ACEEffectorActor, Mode), &ACEEffectorActor::OnModeChanged },
+	{ GET_MEMBER_NAME_CHECKED(ACEEffectorActor, Offset), &ACEEffectorActor::OnTransformOptionsChanged },
+	{ GET_MEMBER_NAME_CHECKED(ACEEffectorActor, Rotation), &ACEEffectorActor::OnTransformOptionsChanged },
+	{ GET_MEMBER_NAME_CHECKED(ACEEffectorActor, Scale), &ACEEffectorActor::OnTransformOptionsChanged },
+	{ GET_MEMBER_NAME_CHECKED(ACEEffectorActor, TargetActorWeak), &ACEEffectorActor::OnTargetActorChanged },
+	{ GET_MEMBER_NAME_CHECKED(ACEEffectorActor, LocationStrength), &ACEEffectorActor::OnNoiseFieldOptionsChanged },
+	{ GET_MEMBER_NAME_CHECKED(ACEEffectorActor, RotationStrength), &ACEEffectorActor::OnNoiseFieldOptionsChanged },
+	{ GET_MEMBER_NAME_CHECKED(ACEEffectorActor, ScaleStrength), &ACEEffectorActor::OnNoiseFieldOptionsChanged },
+	{ GET_MEMBER_NAME_CHECKED(ACEEffectorActor, Pan), &ACEEffectorActor::OnNoiseFieldOptionsChanged },
+	{ GET_MEMBER_NAME_CHECKED(ACEEffectorActor, Frequency), &ACEEffectorActor::OnNoiseFieldOptionsChanged },
+	/** Force */
+	{ GET_MEMBER_NAME_CHECKED(ACEEffectorActor, bOrientationForceEnabled), &ACEEffectorActor::OnForceEnabledChanged },
+	{ GET_MEMBER_NAME_CHECKED(ACEEffectorActor, OrientationForceRate), &ACEEffectorActor::OnForceOptionsChanged },
+	{ GET_MEMBER_NAME_CHECKED(ACEEffectorActor, OrientationForceMin), &ACEEffectorActor::OnForceOptionsChanged },
+	{ GET_MEMBER_NAME_CHECKED(ACEEffectorActor, OrientationForceMax), &ACEEffectorActor::OnForceOptionsChanged },
+	{ GET_MEMBER_NAME_CHECKED(ACEEffectorActor, bVortexForceEnabled), &ACEEffectorActor::OnForceEnabledChanged },
+	{ GET_MEMBER_NAME_CHECKED(ACEEffectorActor, VortexForceAmount), &ACEEffectorActor::OnForceOptionsChanged },
+	{ GET_MEMBER_NAME_CHECKED(ACEEffectorActor, VortexForceAxis), &ACEEffectorActor::OnForceOptionsChanged },
+	{ GET_MEMBER_NAME_CHECKED(ACEEffectorActor, bCurlNoiseForceEnabled), &ACEEffectorActor::OnForceEnabledChanged },
+	{ GET_MEMBER_NAME_CHECKED(ACEEffectorActor, CurlNoiseForceStrength), &ACEEffectorActor::OnForceOptionsChanged },
+	{ GET_MEMBER_NAME_CHECKED(ACEEffectorActor, CurlNoiseForceFrequency), &ACEEffectorActor::OnForceOptionsChanged },
+	{ GET_MEMBER_NAME_CHECKED(ACEEffectorActor, bAttractionForceEnabled), &ACEEffectorActor::OnForceEnabledChanged },
+	{ GET_MEMBER_NAME_CHECKED(ACEEffectorActor, AttractionForceStrength), &ACEEffectorActor::OnForceOptionsChanged },
+	{ GET_MEMBER_NAME_CHECKED(ACEEffectorActor, AttractionForceFalloff), &ACEEffectorActor::OnForceOptionsChanged },
+	{ GET_MEMBER_NAME_CHECKED(ACEEffectorActor, bGravityForceEnabled), &ACEEffectorActor::OnForceEnabledChanged },
+	{ GET_MEMBER_NAME_CHECKED(ACEEffectorActor, GravityForceAcceleration), &ACEEffectorActor::OnForceOptionsChanged },
+};
+
+void ACEEffectorActor::PostEditChangeProperty(FPropertyChangedEvent& InPropertyChangedEvent)
+{
+	Super::PostEditChangeProperty(InPropertyChangedEvent);
+
+	PropertyChangeDispatcher.OnPropertyChanged(this, InPropertyChangedEvent);
+}
+
+void ACEEffectorActor::PostEditImport()
+{
+	Super::PostEditImport();
+
+	RegisterToCloners();
+}
+
+void ACEEffectorActor::PostEditUndo()
+{
+	Super::PostEditUndo();
+
+	RegisterToChannel();
+
+	OnClonersChanged();
+
+	// Undo causes cloner to blackout, so refresh state
+	ForEachCloner([](ACEClonerActor* InCloner, int32 InIndex)
+	{
+		InCloner->ForceUpdateCloner();
+		return true;
+	});
+}
+#endif
+
+void ACEEffectorActor::Destroyed()
+{
+	Super::Destroyed();
+
+	// Remove this effector from the effector channel
+	if (UCEEffectorSubsystem* EffectorSubsystem = UCEEffectorSubsystem::Get(GetWorld()))
+	{
+		EffectorSubsystem->UnregisterChannelEffector(this);
+	}
+
+	// Remove this effector from the cloners it is linked to
+	ForEachCloner([this](ACEClonerActor* InCloner, int32 InIndex)
+	{
+		InCloner->UnregisterEffector(this);
+		return true;
+	});
+}
+
+void ACEEffectorActor::PostActorCreated()
+{
+	Super::PostActorCreated();
+
+	OnEffectorChanged();
+}
+
+void ACEEffectorActor::PostLoad()
+{
+	Super::PostLoad();
+
+	RegisterToCloners();
+}
+
+void ACEEffectorActor::PostDuplicate(EDuplicateMode::Type InDuplicateMode)
+{
+	Super::PostDuplicate(InDuplicateMode);
+
+	RegisterToCloners();
+}
+
+void ACEEffectorActor::LinkCloner(ACEClonerActor* InCloner)
+{
+	if (!InCloner || Cloners.Contains(InCloner))
+	{
+		return;
+	}
+
+	Cloners.Add(InCloner);
+	OnClonersChanged();
+}
+
+void ACEEffectorActor::UnlinkCloner(ACEClonerActor* InCloner)
+{
+	if (!InCloner || !Cloners.Contains(InCloner))
+	{
+		return;
+	}
+
+	Cloners.Remove(InCloner);
+	OnClonersChanged();
+}
+
+void ACEEffectorActor::OnEffectorIdentifierChanged()
+{
+	int32 Identifier = ChannelData.GetIdentifier();
+
+	if (Identifier == INDEX_NONE)
+	{
+		UE_LOG(LogCEEffectorActor, Log, TEXT("%s : Effector channel identifier is invalid"), *GetActorNameOrLabel());
+		return;
+	}
+
+	// Update effector identifier in linked cloners
+	ForEachCloner([Identifier](ACEClonerActor* InCloner, int32 InEffectorIndex)
+	{
+		if (const FCEClonerEffectorDataInterfaces* DataInterfaces = InCloner->GetEffectorDataInterfaces())
+		{
+			UNiagaraDataInterfaceArrayInt32* EffectorIndexDI = DataInterfaces->GetIndexArray();
+
+			// This effector has a specific index in the cloner
+			if (EffectorIndexDI->GetArrayReference().IsValidIndex(InEffectorIndex))
+			{
+				const int32 PrevIdentifier = EffectorIndexDI->GetArrayReference()[InEffectorIndex];
+
+				// Only update if identifier is different
+				if (PrevIdentifier != Identifier)
+				{
+					EffectorIndexDI->GetArrayReference()[InEffectorIndex] = Identifier;
+
+					constexpr bool bImmediateUpdate = true;
+					InCloner->RequestClonerUpdate(bImmediateUpdate);
+				}
+			}
+		}
+
+		return true;
+	});
+}
+
+void ACEEffectorActor::SetType(ECEClonerEffectorType InType)
+{
+	if (InType == Type)
+	{
+		return;
+	}
+
+	Type = InType;
+	OnTypeChanged();
+}
+
+void ACEEffectorActor::SetEasing(ECEClonerEasing InEasing)
+{
+	if (InEasing == Easing)
+	{
+		return;
+	}
+
+	Easing = InEasing;
+	OnEasingChanged();
+}
+
+void ACEEffectorActor::SetMagnitude(float InMagnitude)
+{
+	if (FMath::IsNearlyEqual(InMagnitude, Magnitude))
+	{
+		return;
+	}
+
+	if (InMagnitude < 0.f || InMagnitude > 1.f)
+	{
+		return;
+	}
+
+	Magnitude = InMagnitude;
+	OnMagnitudeChanged();
+}
+
+void ACEEffectorActor::SetOuterRadius(float InRadius)
+{
+	if (FMath::IsNearlyEqual(InRadius, OuterRadius))
+	{
+		return;
+	}
+
+	if (InRadius < 0.f || InRadius < InnerRadius)
+	{
+		return;
+	}
+
+	OuterRadius = InRadius;
+	OnSphereChanged();
+}
+
+void ACEEffectorActor::SetInnerRadius(float InRadius)
+{
+	if (FMath::IsNearlyEqual(InRadius, InnerRadius))
+	{
+		return;
+	}
+
+	if (InRadius < 0.f || InRadius > OuterRadius)
+	{
+		return;
+	}
+
+	InnerRadius = InRadius;
+	OnSphereChanged();
+}
+
+void ACEEffectorActor::SetInnerExtent(const FVector& InExtent)
+{
+	if (InExtent.Equals(InnerExtent))
+	{
+		return;
+	}
+
+	if (InExtent.X < 0 || InExtent.Y < 0 || InExtent.Z < 0)
+	{
+		return;
+	}
+
+	InnerExtent = InExtent;
+	OnBoxChanged();
+}
+
+void ACEEffectorActor::SetOuterExtent(const FVector& InExtent)
+{
+	if (InExtent.Equals(OuterExtent))
+	{
+		return;
+	}
+
+	if (InExtent.X < 0 || InExtent.Y < 0 || InExtent.Z < 0)
+	{
+		return;
+	}
+
+	OuterExtent = InExtent;
+	OnBoxChanged();
+}
+
+void ACEEffectorActor::SetPlaneSpacing(float InSpacing)
+{
+	if (FMath::IsNearlyEqual(InSpacing, PlaneSpacing))
+	{
+		return;
+	}
+
+	if (InSpacing < 0.f)
+	{
+		return;
+	}
+
+	PlaneSpacing = InSpacing;
+	OnPlaneChanged();
+}
+
+void ACEEffectorActor::SetOffset(const FVector& InOffset)
+{
+	if (InOffset.Equals(Offset))
+	{
+		return;
+	}
+
+	Offset = InOffset;
+	OnTransformOptionsChanged();
+}
+
+void ACEEffectorActor::SetRotation(const FRotator& InRotation)
+{
+	if (InRotation.Equals(Rotation))
+	{
+		return;
+	}
+
+	constexpr float MinRotation = -180.f;
+	if (InRotation.Pitch < MinRotation || InRotation.Roll < MinRotation || InRotation.Yaw < MinRotation)
+	{
+		return;
+	}
+
+	constexpr float MaxRotation = 180.f;
+	if (InRotation.Pitch > MaxRotation || InRotation.Roll > MaxRotation || InRotation.Yaw > MaxRotation)
+	{
+		return;
+	}
+
+	Rotation = InRotation;
+	OnTransformOptionsChanged();
+}
+
+void ACEEffectorActor::SetScale(const FVector& InScale)
+{
+	if (InScale.Equals(Scale))
+	{
+		return;
+	}
+
+	if (InScale.X < 0.f || InScale.Y < 0.f || InScale.Z < 0.f)
+	{
+		return;
+	}
+
+	Scale = InScale;
+	OnTransformOptionsChanged();
+}
+
+void ACEEffectorActor::SetEnabled(bool bInEnable)
+{
+	if (bInEnable == bEnabled)
+	{
+		return;
+	}
+
+	bEnabled = bInEnable;
+	OnEnabledChanged();
+}
+
+void ACEEffectorActor::SetVisualizerThickness(float InThickness)
+{
+	if (FMath::IsNearlyEqual(VisualizerThickness, InThickness))
+	{
+		return;
+	}
+
+	if (VisualizerThickness < 0.1f || VisualizerThickness > 10.f)
+	{
+		return;
+	}
+
+	VisualizerThickness = InThickness;
+	OnVisualizerThicknessChanged();
+}
+
+void ACEEffectorActor::SetMode(ECEClonerEffectorMode InMode)
+{
+	if (Mode == InMode)
+	{
+		return;
+	}
+
+	Mode = InMode;
+	OnModeChanged();
+}
+
+void ACEEffectorActor::SetTargetActor(AActor* InTargetActor)
+{
+	if (InTargetActor == TargetActorWeak.Get())
+	{
+		return;
+	}
+
+	TargetActorWeak = InTargetActor;
+	OnTargetActorChanged();
+}
+
+void ACEEffectorActor::SetAttractionForceStrength(float InForceStrength)
+{
+	if (FMath::IsNearlyEqual(AttractionForceStrength, InForceStrength))
+	{
+		return;
+	}
+
+	AttractionForceStrength = InForceStrength;
+	OnForceOptionsChanged();
+}
+
+void ACEEffectorActor::SetAttractionForceFalloff(float InForceFalloff)
+{
+	if (FMath::IsNearlyEqual(AttractionForceFalloff, InForceFalloff))
+	{
+		return;
+	}
+
+	AttractionForceFalloff = InForceFalloff;
+	OnForceOptionsChanged();
+}
+
+void ACEEffectorActor::SetGravityForceEnabled(bool bInForceEnabled)
+{
+	if (bGravityForceEnabled == bInForceEnabled)
+	{
+		return;
+	}
+
+	bGravityForceEnabled = bInForceEnabled;
+	OnForceEnabledChanged();
+}
+
+void ACEEffectorActor::SetGravityForceAcceleration(const FVector& InAcceleration)
+{
+	if (GravityForceAcceleration.Equals(InAcceleration))
+	{
+		return;
+	}
+
+	GravityForceAcceleration = InAcceleration;
+	OnForceOptionsChanged();
+}
+
+void ACEEffectorActor::SetLocationStrength(const FVector& InStrength)
+{
+	if (LocationStrength.Equals(InStrength))
+	{
+		return;
+	}
+
+	LocationStrength = InStrength;
+	OnNoiseFieldOptionsChanged();
+}
+
+void ACEEffectorActor::SetRotationStrength(const FRotator& InStrength)
+{
+	if (RotationStrength.Equals(InStrength))
+	{
+		return;
+	}
+
+	RotationStrength = InStrength;
+	OnNoiseFieldOptionsChanged();
+}
+
+void ACEEffectorActor::SetScaleStrength(const FVector& InStrength)
+{
+	if (ScaleStrength.Equals(InStrength))
+	{
+		return;
+	}
+
+	ScaleStrength = InStrength;
+	OnNoiseFieldOptionsChanged();
+}
+
+void ACEEffectorActor::SetPan(const FVector& InPan)
+{
+	if (Pan.Equals(InPan))
+	{
+		return;
+	}
+
+	Pan = InPan;
+	OnNoiseFieldOptionsChanged();
+}
+
+void ACEEffectorActor::SetFrequency(float InFrequency)
+{
+	if (FMath::IsNearlyEqual(Frequency, InFrequency))
+	{
+		return;
+	}
+
+	if (InFrequency < 0.f)
+	{
+		return;
+	}
+
+	Frequency = InFrequency;
+	OnNoiseFieldOptionsChanged();
+}
+
+void ACEEffectorActor::SetOrientationForceEnabled(bool bInForceEnabled)
+{
+	if (bOrientationForceEnabled == bInForceEnabled)
+	{
+		return;
+	}
+
+	bOrientationForceEnabled = bInForceEnabled;
+	OnForceEnabledChanged();
+}
+
+void ACEEffectorActor::SetOrientationForceRate(float InForceOrientationRate)
+{
+	if (FMath::IsNearlyEqual(OrientationForceRate, InForceOrientationRate))
+	{
+		return;
+	}
+
+	OrientationForceRate = InForceOrientationRate;
+	OnForceOptionsChanged();
+}
+
+void ACEEffectorActor::SetOrientationForceMin(const FVector& InForceOrientationMin)
+{
+	if (OrientationForceMin.Equals(InForceOrientationMin))
+	{
+		return;
+	}
+
+	OrientationForceMin = InForceOrientationMin;
+	OnForceOptionsChanged();
+}
+
+void ACEEffectorActor::SetOrientationForceMax(const FVector& InForceOrientationMax)
+{
+	if (OrientationForceMax.Equals(InForceOrientationMax))
+	{
+		return;
+	}
+
+	OrientationForceMax = InForceOrientationMax;
+	OnForceOptionsChanged();
+}
+
+void ACEEffectorActor::SetVortexForceEnabled(bool bInForceEnabled)
+{
+	if (bVortexForceEnabled == bInForceEnabled)
+	{
+		return;
+	}
+
+	bVortexForceEnabled = bInForceEnabled;
+	OnForceEnabledChanged();
+}
+
+void ACEEffectorActor::SetVortexForceAmount(float InForceVortexAmount)
+{
+	if (FMath::IsNearlyEqual(VortexForceAmount, InForceVortexAmount))
+	{
+		return;
+	}
+
+	VortexForceAmount = InForceVortexAmount;
+	OnForceOptionsChanged();
+}
+
+void ACEEffectorActor::SetVortexForceAxis(const FVector& InForceVortexAxis)
+{
+	if (VortexForceAxis.Equals(InForceVortexAxis))
+	{
+		return;
+	}
+
+	VortexForceAxis = InForceVortexAxis;
+	OnForceOptionsChanged();
+}
+
+void ACEEffectorActor::SetCurlNoiseForceEnabled(bool bInForceEnabled)
+{
+	if (bCurlNoiseForceEnabled == bInForceEnabled)
+	{
+		return;
+	}
+
+	bCurlNoiseForceEnabled = bInForceEnabled;
+	OnForceEnabledChanged();
+}
+
+void ACEEffectorActor::SetCurlNoiseForceStrength(float InForceCurlNoiseStrength)
+{
+	if (FMath::IsNearlyEqual(CurlNoiseForceStrength, InForceCurlNoiseStrength))
+	{
+		return;
+	}
+
+	CurlNoiseForceStrength = InForceCurlNoiseStrength;
+	OnForceOptionsChanged();
+}
+
+void ACEEffectorActor::SetCurlNoiseForceFrequency(float InForceCurlNoiseFrequency)
+{
+	if (FMath::IsNearlyEqual(CurlNoiseForceFrequency, InForceCurlNoiseFrequency))
+	{
+		return;
+	}
+
+	CurlNoiseForceFrequency = InForceCurlNoiseFrequency;
+	OnForceOptionsChanged();
+}
+
+void ACEEffectorActor::SetAttractionForceEnabled(bool bInForceEnabled)
+{
+	if (bAttractionForceEnabled == bInForceEnabled)
+	{
+		return;
+	}
+
+	bAttractionForceEnabled = bInForceEnabled;
+	OnForceEnabledChanged();
+}
+
+#if WITH_EDITOR
+void ACEEffectorActor::SetVisualizerSpriteVisible(bool bInVisible)
+{
+	if (bVisualizerSpriteVisible == bInVisible)
+	{
+		return;
+	}
+
+	bVisualizerSpriteVisible = bInVisible;
+	OnVisualizerSpriteVisibleChanged();
+}
+#endif
+
+void ACEEffectorActor::ForEachCloner(TFunctionRef<bool(ACEClonerActor*, int32)> InFunction, bool bSkipIdxCheck)
+{
+	for (TSet<TWeakObjectPtr<ACEClonerActor>>::TIterator It(InternalCloners); It; ++It)
+	{
+		ACEClonerActor* Cloner = It->Get();
+		if (!Cloner)
+		{
+			continue;
+		}
+
+		const int32 Idx = Cloner->GetEffectorIndex(this);
+		if (!bSkipIdxCheck && Idx == INDEX_NONE)
+		{
+			continue;
+		}
+
+		if (!InFunction(Cloner, Idx))
+		{
+			return;
+		}
+	}
+}
+
+void ACEEffectorActor::OnEffectorSubsystemInitialized(const UWorld* World)
+{
+	if (GetWorld() == World)
+	{
+		RegisterToChannel();
+	}
+}
+
+void ACEEffectorActor::RegisterToChannel()
+{
+	UCEEffectorSubsystem::OnSubsystemInitializedDelegate.RemoveAll(this);
+
+	// Register this effector to the effector channel
+	if (UCEEffectorSubsystem* EffectorSubsystem = UCEEffectorSubsystem::Get(GetWorld()))
+	{
+		EffectorSubsystem->RegisterChannelEffector(this);
+	}
+	else
+	{
+		UCEEffectorSubsystem::OnSubsystemInitializedDelegate.AddUObject(this, &ACEEffectorActor::OnEffectorSubsystemInitialized);
+	}
+}
+
+void ACEEffectorActor::OnEffectorTransformed(USceneComponent* InUpdatedComponent, EUpdateTransformFlags InUpdateTransformFlags, ETeleportType InTeleport)
+{
+	OnTransformChanged();
+	// Update when scaled or rotated
+	OnSphereChanged();
+	OnBoxChanged();
+	OnPlaneChanged();
+	// update if self
+	const AActor* InternalTargetActor = InternalTargetActorWeak.Get();
+	if (InternalTargetActor == this)
+	{
+		OnTargetOptionsChanged();
+	}
+}
+
+void ACEEffectorActor::OnEnabledChanged()
+{
+	if (bEnabled)
+	{
+		OnEffectorChanged();
+	}
+	else // Disabled
+	{
+		OnEffectorDisabled();
+		// Hide visualization components
+		OnTypeChanged();
+		// Editor
+		OnVisualizerThicknessChanged();
+		OnVisualizerSpriteVisibleChanged();
+	}
+}
+
+void ACEEffectorActor::OnEffectorDisabled()
+{
+	OnMagnitudeChanged();
+}
+
+void ACEEffectorActor::OnClonersChanged()
+{
+	TSet<TWeakObjectPtr<ACEClonerActor>> AddedCloners = Cloners.Difference(InternalCloners);
+	TSet<TWeakObjectPtr<ACEClonerActor>> RemovedCloners = InternalCloners.Difference(Cloners);
+
+	// Unlink removed cloners
+	for (const TWeakObjectPtr<ACEClonerActor>& RemovedClonerWeak : RemovedCloners)
+	{
+		if (ACEClonerActor* RemovedCloner = RemovedClonerWeak.Get())
+		{
+			RemovedCloner->UnregisterEffector(this);
+		}
+	}
+
+	// Link added cloners
+	for (const TWeakObjectPtr<ACEClonerActor>& AddedClonerWeak : AddedCloners)
+	{
+		if (ACEClonerActor* AddedCloner = AddedClonerWeak.Get())
+		{
+			AddedCloner->RegisterEffector(this);
+		}
+	}
+}
+
+void ACEEffectorActor::OnModeChanged()
+{
+	ChannelData.Mode = Mode;
+
+	OnTransformOptionsChanged();
+	OnTargetActorChanged();
+	OnNoiseFieldOptionsChanged();
+}
+
+void ACEEffectorActor::OnTypeChanged()
+{
+	// Sphere
+	if (InnerSphereComponent && OuterSphereComponent)
+	{
+		InnerSphereComponent->SetVisibility(Type == ECEClonerEffectorType::Sphere);
+		OuterSphereComponent->SetVisibility(Type == ECEClonerEffectorType::Sphere);
+	}
+
+	// Box
+	if (InnerBoxComponent && OuterBoxComponent)
+	{
+		InnerBoxComponent->SetVisibility(Type == ECEClonerEffectorType::Box);
+		OuterBoxComponent->SetVisibility(Type == ECEClonerEffectorType::Box);
+	}
+
+	// Plane
+	if (InnerPlaneComponent && OuterPlaneComponent)
+	{
+		InnerPlaneComponent->SetVisibility(Type == ECEClonerEffectorType::Plane);
+		OuterPlaneComponent->SetVisibility(Type == ECEClonerEffectorType::Plane);
+	}
+
+	ChannelData.Type = Type;
+
+	// Update type data
+	OnSphereChanged();
+	OnBoxChanged();
+	OnPlaneChanged();
+}
+
+void ACEEffectorActor::RefreshClonerParameters(ACEClonerActor* InCloner, bool bInForce) const
+{
+	if (InCloner && (bEnabled || bInForce))
+	{
+		InCloner->RequestClonerUpdate();
+	}
+}
+
+void ACEEffectorActor::RefreshClonersParameters(bool bInForce)
+{
+	ForEachCloner([this, bInForce](ACEClonerActor* InCloner, int32 InIndex)
+	{
+		RefreshClonerParameters(InCloner, bInForce);
+		return true;
+	});
+}
+
+void ACEEffectorActor::RegisterToCloners()
+{
+	for (TSet<TWeakObjectPtr<ACEClonerActor>>::TIterator It = InternalCloners.CreateIterator(); It; ++It)
+	{
+		ACEClonerActor* Cloner = It->Get();
+
+		if (!Cloner)
+		{
+			It.RemoveCurrent();
+			continue;
+		}
+
+		Cloner->RegisterEffector(this);
+	}
+
+	OnEffectorChanged();
+}
+
+void ACEEffectorActor::OnEffectorChanged()
+{
+	if (!bEnabled)
+	{
+		OnEnabledChanged();
+	}
+	else
+	{
+		OnModeChanged();
+		OnTypeChanged();
+		OnEasingChanged();
+		OnTransformChanged();
+		OnMagnitudeChanged();
+		OnForceOptionsChanged();
+		// Editor
+		OnVisualizerThicknessChanged();
+		OnVisualizerSpriteVisibleChanged();
+	}
+}
+
+void ACEEffectorActor::OnClonerLinked(ACEClonerActor* InCloner, int32 InEffectorIdx)
+{
+	if (!InCloner)
+	{
+		return;
+	}
+
+	InternalCloners.Add(InCloner);
+
+	OnEffectorIdentifierChanged();
+
+	UE_LOG(LogCEEffectorActor, Log, TEXT("%s : Effector linked to Cloner %s"), *GetActorNameOrLabel(), *InCloner->GetActorNameOrLabel());
+}
+
+void ACEEffectorActor::OnClonerUnlinked(ACEClonerActor* InCloner, int32 InEffectorIdx)
+{
+	if (!InCloner)
+	{
+		return;
+	}
+
+	InternalCloners.Remove(InCloner);
+
+	UE_LOG(LogCEEffectorActor, Log, TEXT("%s : Effector unlinked from cloner %s"), *GetActorNameOrLabel(), *InCloner->GetActorNameOrLabel());
+}
+
+void ACEEffectorActor::OnEasingChanged()
+{
+	ChannelData.Easing = Easing;
+}
+
+void ACEEffectorActor::OnTransformChanged()
+{
+	ChannelData.Location = GetActorLocation();
+	ChannelData.Rotation = GetActorRotation().Quaternion();
+	ChannelData.Scale = GetActorScale();
+}
+
+void ACEEffectorActor::OnBoxChanged()
+{
+	InnerExtent = ClampVector(InnerExtent, FVector::ZeroVector, OuterExtent);
+	// Max()
+	OuterExtent = OuterExtent.ComponentMax(InnerExtent);
+
+	if (InnerBoxComponent && OuterBoxComponent)
+	{
+		InnerBoxComponent->SetBoxExtent(InnerExtent);
+		OuterBoxComponent->SetBoxExtent(OuterExtent);
+	}
+
+	if (Type != ECEClonerEffectorType::Box)
+	{
+		return;
+	}
+
+	ChannelData.InnerExtent = GetInnerExtent();
+	ChannelData.OuterExtent = GetOuterExtent();
+}
+
+void ACEEffectorActor::OnSphereChanged()
+{
+	InnerRadius = FMath::Clamp(InnerRadius, 0.f, OuterRadius);
+	OuterRadius = FMath::Max(OuterRadius, InnerRadius);
+
+	if (InnerSphereComponent && OuterSphereComponent)
+	{
+		InnerSphereComponent->SetSphereRadius(InnerRadius);
+		OuterSphereComponent->SetSphereRadius(OuterRadius);
+	}
+
+	if (Type != ECEClonerEffectorType::Sphere)
+	{
+		return;
+	}
+
+	ChannelData.InnerExtent = FVector(GetInnerRadius());
+	ChannelData.OuterExtent = FVector(GetOuterRadius());
+}
+
+void ACEEffectorActor::OnPlaneChanged()
+{
+	static const FVector PlaneAxis = -FVector::LeftVector;
+	const FVector InnerPlane = PlaneAxis * FVector(-PlaneSpacing/2);
+	const FVector OuterPlane = PlaneAxis * FVector(PlaneSpacing/2);
+
+	if (InnerPlaneComponent && OuterPlaneComponent)
+	{
+		static const FVector InnerSize(100, 0, 100);
+		InnerPlaneComponent->SetRelativeLocation(InnerPlane);
+		InnerPlaneComponent->SetBoxExtent(InnerSize);
+
+		static const FVector OuterSize(200, 0, 200);
+		OuterPlaneComponent->SetRelativeLocation(OuterPlane);
+		OuterPlaneComponent->SetBoxExtent(OuterSize);
+	}
+
+	if (Type != ECEClonerEffectorType::Plane)
+	{
+		return;
+	}
+
+	ChannelData.InnerExtent = PlaneAxis;
+	ChannelData.OuterExtent = FVector(PlaneSpacing);
+}
+
+void ACEEffectorActor::OnMagnitudeChanged()
+{
+	const float EffectorMagnitude = bEnabled ? GetMagnitude() : 0.f;
+	ChannelData.Magnitude = EffectorMagnitude;
+}
+
+void ACEEffectorActor::OnTransformOptionsChanged()
+{
+	if (Mode != ECEClonerEffectorMode::Default)
+	{
+		return;
+	}
+
+	FVector ScaleDelta = GetScale();
+	ScaleDelta.X = FMath::Max(ScaleDelta.X, UE_KINDA_SMALL_NUMBER);
+	ScaleDelta.Y = FMath::Max(ScaleDelta.Y, UE_KINDA_SMALL_NUMBER);
+	ScaleDelta.Z = FMath::Max(ScaleDelta.Z, UE_KINDA_SMALL_NUMBER);
+
+	ChannelData.LocationDelta = GetOffset();
+	ChannelData.RotationDelta = GetRotation().Quaternion();
+	ChannelData.ScaleDelta = ScaleDelta;
+}
+
+void ACEEffectorActor::OnTargetActorChanged()
+{
+	AActor* TargetActor = TargetActorWeak.Get();
+	AActor* const InternalTargetActor = InternalTargetActorWeak.Get();
+	if (TargetActor && TargetActor == InternalTargetActor)
+	{
+		OnTargetOptionsChanged();
+		return;
+	}
+
+	// unbind, except if it is self
+	if (InternalTargetActor && InternalTargetActor->GetRootComponent())
+	{
+		if (InternalTargetActor != this)
+		{
+			InternalTargetActor->OnDestroyed.RemoveAll(this);
+			InternalTargetActor->GetRootComponent()->TransformUpdated.RemoveAll(this);
+			InternalTargetActorWeak.Reset();
+		}
+	}
+
+	// set self by default if invalid
+	if (!TargetActor || !TargetActor->GetRootComponent())
+	{
+		TargetActor = this;
+	}
+
+	// bind to transform event, do not bind to self since we already do that
+	if (TargetActor != this)
+	{
+		TargetActor->GetRootComponent()->TransformUpdated.RemoveAll(this);
+		TargetActor->GetRootComponent()->TransformUpdated.AddUObject(this, &ACEEffectorActor::OnTargetActorTransformChanged);
+		TargetActor->OnDestroyed.RemoveAll(this);
+		TargetActor->OnDestroyed.AddUniqueDynamic(this, &ACEEffectorActor::OnTargetActorDestroyed);
+	}
+
+	TargetActorWeak = TargetActor;
+	InternalTargetActorWeak = TargetActor;
+	OnTargetOptionsChanged();
+}
+
+void ACEEffectorActor::OnTargetActorDestroyed(AActor* InActor)
+{
+	const AActor* TargetActor = TargetActorWeak.Get();
+	if (TargetActor == InActor)
+	{
+		TargetActorWeak = this;
+		OnTargetActorChanged();
+	}
+}
+
+void ACEEffectorActor::OnNoiseFieldOptionsChanged()
+{
+	if (Mode != ECEClonerEffectorMode::NoiseField)
+	{
+		return;
+	}
+
+	ChannelData.LocationDelta = LocationStrength;
+	ChannelData.RotationDelta = RotationStrength.Quaternion();
+	ChannelData.ScaleDelta = ScaleStrength;
+	ChannelData.Frequency = Frequency;
+	ChannelData.Pan = Pan;
+}
+
+void ACEEffectorActor::OnTargetActorTransformChanged(USceneComponent*, EUpdateTransformFlags, ETeleportType)
+{
+	OnTargetOptionsChanged();
+}
+
+void ACEEffectorActor::OnTargetOptionsChanged()
+{
+	if (Mode != ECEClonerEffectorMode::Target)
+	{
+		return;
+	}
+
+	const AActor* InternalTargetActor = InternalTargetActorWeak.Get();
+	if (!InternalTargetActor)
+	{
+		return;
+	}
+
+	ChannelData.LocationDelta = InternalTargetActor->GetActorLocation();
+	ChannelData.RotationDelta = FQuat::Identity;
+	ChannelData.ScaleDelta = FVector::OneVector;
+}
+
+void ACEEffectorActor::OnVisualizerThicknessChanged()
+{
+	VisualizerThickness = FMath::Clamp(VisualizerThickness, 0.1f, 10.f);
+
+	if (InnerSphereComponent && OuterSphereComponent)
+	{
+		InnerSphereComponent->SetLineThickness(VisualizerThickness);
+		OuterSphereComponent->SetLineThickness(VisualizerThickness);
+	}
+
+	if (InnerBoxComponent && OuterBoxComponent)
+	{
+		InnerBoxComponent->SetLineThickness(VisualizerThickness);
+		OuterBoxComponent->SetLineThickness(VisualizerThickness);
+	}
+
+	if (InnerPlaneComponent && OuterPlaneComponent)
+	{
+		InnerPlaneComponent->SetLineThickness(VisualizerThickness);
+		OuterPlaneComponent->SetLineThickness(VisualizerThickness);
+	}
+}
+
+void ACEEffectorActor::OnVisualizerSpriteVisibleChanged()
+{
+#if WITH_EDITOR
+	UE::ClonerEffector::SetBillboardComponentSprite(this, TEXT("/Script/Engine.Texture2D'/ClonerEffector/Textures/T_EffectorIcon.T_EffectorIcon'"));
+	UE::ClonerEffector::SetBillboardComponentVisibility(this, bVisualizerSpriteVisible);
+#endif
+}
+
+void ACEEffectorActor::OnForceOptionsChanged()
+{
+	if (bOrientationForceEnabled)
+	{
+		ChannelData.OrientationForceRate = OrientationForceRate;
+		ChannelData.OrientationForceMin = OrientationForceMin;
+		ChannelData.OrientationForceMax = OrientationForceMax;
+	}
+	else
+	{
+		ChannelData.OrientationForceRate = 0.f;
+		ChannelData.OrientationForceMin = FVector::ZeroVector;
+		ChannelData.OrientationForceMax = FVector::ZeroVector;
+	}
+
+	if (bVortexForceEnabled)
+	{
+		ChannelData.VortexForceAmount = VortexForceAmount;
+		ChannelData.VortexForceAxis = VortexForceAxis;
+	}
+	else
+	{
+		ChannelData.VortexForceAmount = 0.f;
+		ChannelData.VortexForceAxis = FVector::ZeroVector;
+	}
+
+	if (bCurlNoiseForceEnabled)
+	{
+		ChannelData.CurlNoiseForceStrength = CurlNoiseForceStrength;
+		ChannelData.CurlNoiseForceFrequency = CurlNoiseForceFrequency;
+	}
+	else
+	{
+		ChannelData.CurlNoiseForceStrength = 0.f;
+		ChannelData.CurlNoiseForceFrequency = 0.f;
+	}
+
+	if (bAttractionForceEnabled)
+	{
+		ChannelData.AttractionForceStrength = AttractionForceStrength;
+		ChannelData.AttractionForceFalloff = AttractionForceFalloff;
+	}
+	else
+	{
+		ChannelData.AttractionForceStrength = 0.f;
+		ChannelData.AttractionForceFalloff = 0.f;
+	}
+
+	if (bGravityForceEnabled)
+	{
+		ChannelData.GravityForceAcceleration = GravityForceAcceleration;
+	}
+	else
+	{
+		ChannelData.GravityForceAcceleration = FVector::ZeroVector;
+	}
+}
+
+void ACEEffectorActor::OnForceEnabledChanged()
+{
+	// Refresh cloners to reset clones transform
+	RefreshClonersParameters();
+	OnForceOptionsChanged();
+}
diff --git a/Engine/Plugins/Experimental/ClonerEffector/Source/ClonerEffector/Private/Subsystems/CEClonerSubsystem.cpp b/Engine/Plugins/Experimental/ClonerEffector/Source/ClonerEffector/Private/Subsystems/CEClonerSubsystem.cpp
new file mode 100644
--- /dev/null
+++ b/Engine/Plugins/Experimental/ClonerEffector/Source/ClonerEffector/Private/Subsystems/CEClonerSubsystem.cpp
@@ -0,0 +1,263 @@
+// Copyright Epic Games, Inc. All Rights Reserved.
+
+#include "Subsystems/CEClonerSubsystem.h"
+
+#include "Cloner/CEClonerActor.h"
+#include "Cloner/Layouts/CEClonerCircleLayout.h"
+#include "Cloner/Layouts/CEClonerCylinderLayout.h"
+#include "Cloner/Layouts/CEClonerGridLayout.h"
+#include "Cloner/Layouts/CEClonerHoneycombLayout.h"
+#include "Cloner/Layouts/CEClonerLayoutBase.h"
+#include "Cloner/Layouts/CEClonerLineLayout.h"
+#include "Cloner/Layouts/CEClonerMeshLayout.h"
+#include "Cloner/Layouts/CEClonerSphereRandomLayout.h"
+#include "Cloner/Layouts/CEClonerSphereUniformLayout.h"
+#include "Cloner/Layouts/CEClonerSplineLayout.h"
+#include "Engine/Engine.h"
+#include "UObject/Class.h"
+#include "UObject/UObjectIterator.h"
+
+#if WITH_EDITOR
+#include "HAL/IConsoleManager.h"
+
+UCEClonerSubsystem::FOnCVarChanged UCEClonerSubsystem::OnCVarChangedDelegate;
+#endif
+
+UCEClonerSubsystem::FOnSubsystemInitialized UCEClonerSubsystem::OnSubsystemInitializedDelegate;
+
+UCEClonerSubsystem* UCEClonerSubsystem::Get()
+{
+	if (GEngine)
+	{
+		return GEngine->GetEngineSubsystem<UCEClonerSubsystem>();
+	}
+
+	return nullptr;
+}
+
+void UCEClonerSubsystem::Initialize(FSubsystemCollectionBase& Collection)
+{
+	Super::Initialize(Collection);
+
+	// Register them here to match old order of layout enum
+	RegisterLayoutClass(UCEClonerGridLayout::StaticClass());
+	RegisterLayoutClass(UCEClonerLineLayout::StaticClass());
+	RegisterLayoutClass(UCEClonerCircleLayout::StaticClass());
+	RegisterLayoutClass(UCEClonerCylinderLayout::StaticClass());
+	RegisterLayoutClass(UCEClonerSphereUniformLayout::StaticClass());
+	RegisterLayoutClass(UCEClonerHoneycombLayout::StaticClass());
+	RegisterLayoutClass(UCEClonerMeshLayout::StaticClass());
+	RegisterLayoutClass(UCEClonerSplineLayout::StaticClass());
+	RegisterLayoutClass(UCEClonerSphereRandomLayout::StaticClass());
+
+	// Scan for new layouts
+	ScanForRegistrableClasses();
+
+#if WITH_EDITOR
+	CVarTSRShadingRejectionFlickeringPeriod = IConsoleManager::Get().FindConsoleVariable(TEXT("r.TSR.ShadingRejection.Flickering.Period"));
+
+	if (CVarTSRShadingRejectionFlickeringPeriod)
+	{
+		CVarTSRShadingRejectionFlickeringPeriod->OnChangedDelegate().AddUObject(this, &UCEClonerSubsystem::OnTSRShadingRejectionFlickeringPeriodChanged);
+	}
+#endif
+
+	OnSubsystemInitializedDelegate.Broadcast();
+}
+
+void UCEClonerSubsystem::Deinitialize()
+{
+	Super::Deinitialize();
+
+#if WITH_EDITOR
+	if (CVarTSRShadingRejectionFlickeringPeriod)
+	{
+		CVarTSRShadingRejectionFlickeringPeriod->OnChangedDelegate().RemoveAll(this);
+	}
+#endif
+}
+
+bool UCEClonerSubsystem::RegisterLayoutClass(const UClass* InClonerLayoutClass)
+{
+	if (!IsValid(InClonerLayoutClass))
+	{
+		return false;
+	}
+
+	if (!InClonerLayoutClass->IsChildOf(UCEClonerLayoutBase::StaticClass())
+		|| InClonerLayoutClass->HasAnyClassFlags(CLASS_Abstract | CLASS_Deprecated | CLASS_NewerVersionExists))
+	{
+		return false;
+	}
+
+	if (IsLayoutClassRegistered(InClonerLayoutClass))
+	{
+		return false;
+	}
+
+	const UCEClonerLayoutBase* CDO = InClonerLayoutClass->GetDefaultObject<UCEClonerLayoutBase>();
+
+	if (!CDO)
+	{
+		return false;
+	}
+
+	// Check niagara asset is valid
+	if (!CDO->IsLayoutValid())
+	{
+		return false;
+	}
+
+	// Does not overwrite existing layouts
+	const FName LayoutName = CDO->GetLayoutName();
+	if (LayoutClasses.Contains(LayoutName))
+	{
+		return false;
+	}
+
+	LayoutClasses.Add(LayoutName, CDO->GetClass());
+
+	return true;
+}
+
+bool UCEClonerSubsystem::UnregisterLayoutClass(const UClass* InClonerLayoutClass)
+{
+	if (!IsValid(InClonerLayoutClass))
+	{
+		return false;
+	}
+
+	for (TMap<FName, TSubclassOf<UCEClonerLayoutBase>>::TIterator It(LayoutClasses); It; ++It)
+	{
+		if (It->Value.Get() == InClonerLayoutClass)
+		{
+			It.RemoveCurrent();
+			return true;
+		}
+	}
+
+	return false;
+}
+
+bool UCEClonerSubsystem::IsLayoutClassRegistered(const UClass* InClonerLayoutClass)
+{
+	if (!IsValid(InClonerLayoutClass))
+	{
+		return false;
+	}
+
+	for (const TPair<FName, TSubclassOf<UCEClonerLayoutBase>>& LayoutClassPair : LayoutClasses)
+	{
+		if (LayoutClassPair.Value.Get() == InClonerLayoutClass)
+		{
+			return true;
+		}
+	}
+
+	return false;
+}
+
+void UCEClonerSubsystem::RegisterCustomActorResolver(FOnGetOrderedActors InCustomResolver)
+{
+	ActorResolver = InCustomResolver;
+}
+
+void UCEClonerSubsystem::UnregisterCustomActorResolver()
+{
+	ActorResolver = FOnGetOrderedActors();
+}
+
+UCEClonerSubsystem::FOnGetOrderedActors& UCEClonerSubsystem::GetCustomActorResolver()
+{
+	return ActorResolver;
+}
+
+TArray<FName> UCEClonerSubsystem::GetLayoutNames() const
+{
+	TArray<FName> LayoutNames;
+	LayoutClasses.GenerateKeyArray(LayoutNames);
+	return LayoutNames;
+}
+
+FName UCEClonerSubsystem::FindLayoutName(TSubclassOf<UCEClonerLayoutBase> InLayoutClass) const
+{
+	if (const FName* Key = LayoutClasses.FindKey(InLayoutClass))
+	{
+		return *Key;
+	}
+
+	return NAME_None;
+}
+
+UCEClonerLayoutBase* UCEClonerSubsystem::CreateNewLayout(FName InLayoutName, ACEClonerActor* InClonerActor)
+{
+	if (!IsValid(InClonerActor))
+	{
+		return nullptr;
+	}
+
+	TSubclassOf<UCEClonerLayoutBase> const* LayoutClass = LayoutClasses.Find(InLayoutName);
+
+	if (!LayoutClass)
+	{
+		return nullptr;
+	}
+
+	return NewObject<UCEClonerLayoutBase>(InClonerActor, LayoutClass->Get());
+}
+
+void UCEClonerSubsystem::ScanForRegistrableClasses()
+{
+	for (const UClass* const Class : TObjectRange<UClass>())
+	{
+		RegisterLayoutClass(Class);
+	}
+}
+
+#if WITH_EDITOR
+void UCEClonerSubsystem::EnableNoFlicker()
+{
+	if (IsNoFlickerEnabled())
+	{
+		return;
+	}
+
+	PreviousCVarValue = CVarTSRShadingRejectionFlickeringPeriod->GetInt();
+	CVarTSRShadingRejectionFlickeringPeriod->Set(NoFlicker);
+}
+
+void UCEClonerSubsystem::DisableNoFlicker()
+{
+	if (!IsNoFlickerEnabled())
+	{
+		return;
+	}
+
+	if (PreviousCVarValue.IsSet())
+	{
+		CVarTSRShadingRejectionFlickeringPeriod->Set(PreviousCVarValue.GetValue());
+	}
+	else
+	{
+		CVarTSRShadingRejectionFlickeringPeriod->Set(*CVarTSRShadingRejectionFlickeringPeriod->GetDefaultValue());
+	}
+}
+
+bool UCEClonerSubsystem::IsNoFlickerEnabled() const
+{
+	if (!CVarTSRShadingRejectionFlickeringPeriod)
+	{
+		return false;
+	}
+
+	return CVarTSRShadingRejectionFlickeringPeriod->GetInt() == NoFlicker;
+}
+
+void UCEClonerSubsystem::OnTSRShadingRejectionFlickeringPeriodChanged(IConsoleVariable* InCVar) const
+{
+	if (InCVar == CVarTSRShadingRejectionFlickeringPeriod)
+	{
+		OnCVarChangedDelegate.Broadcast();
+	}
+}
+#endif
\ No newline at end of file
diff --git a/Engine/Plugins/Experimental/ClonerEffector/Source/ClonerEffector/Private/Subsystems/CEEffectorSubsystem.cpp b/Engine/Plugins/Experimental/ClonerEffector/Source/ClonerEffector/Private/Subsystems/CEEffectorSubsystem.cpp
new file mode 100644
--- /dev/null
+++ b/Engine/Plugins/Experimental/ClonerEffector/Source/ClonerEffector/Private/Subsystems/CEEffectorSubsystem.cpp
@@ -0,0 +1,146 @@
+// Copyright Epic Games, Inc. All Rights Reserved.
+
+#include "Subsystems/CEEffectorSubsystem.h"
+
+#include "Effector/CEEffectorActor.h"
+#include "NiagaraDataChannel.h"
+#include "NiagaraDataChannelAccessor.h"
+#include "NiagaraDataChannelPublic.h"
+
+DEFINE_LOG_CATEGORY_STATIC(LogACEEffectorSubsystem, Log, All);
+
+UCEEffectorSubsystem::FOnSubsystemInitialized UCEEffectorSubsystem::OnSubsystemInitializedDelegate;
+
+UCEEffectorSubsystem* UCEEffectorSubsystem::Get(const UWorld* InWorld)
+{
+	if (InWorld)
+	{
+		return InWorld->GetSubsystem<UCEEffectorSubsystem>();
+	}
+
+	return nullptr;
+}
+
+void UCEEffectorSubsystem::Initialize(FSubsystemCollectionBase& Collection)
+{
+	Super::Initialize(Collection);
+
+	// Load niagara data channel asset for effectors and cache it
+	EffectorDataChannelAsset = LoadObject<UNiagaraDataChannelAsset>(nullptr, DataChannelAssetPath);
+
+	check(EffectorDataChannelAsset->Get());
+}
+
+void UCEEffectorSubsystem::PostInitialize()
+{
+	Super::PostInitialize();
+
+	if (const UWorld* World = GetWorld())
+	{
+		OnSubsystemInitializedDelegate.Broadcast(World);
+	}
+}
+
+bool UCEEffectorSubsystem::RegisterChannelEffector(ACEEffectorActor* InEffector)
+{
+	if (!InEffector)
+	{
+		return false;
+	}
+
+	const int32 EffectorIndex = EffectorsWeak.AddUnique(InEffector);
+	const bool bRegistered = EffectorIndex == EffectorsWeak.Num() - 1;
+
+	if (bRegistered)
+	{
+		UE_LOG(LogACEEffectorSubsystem, Log, TEXT("%s effector registered in channel %i"), *InEffector->GetActorNameOrLabel(), EffectorIndex);
+	}
+
+	return bRegistered;
+}
+
+bool UCEEffectorSubsystem::UnregisterChannelEffector(ACEEffectorActor* InEffector)
+{
+	if (!InEffector)
+	{
+		return false;
+	}
+
+	const bool bUnregistered = EffectorsWeak.Remove(InEffector) > 0;
+
+	if (bUnregistered)
+	{
+		UE_LOG(LogACEEffectorSubsystem, Log, TEXT("%s effector unregistered from channel"), *InEffector->GetActorNameOrLabel());
+	}
+
+	return bUnregistered;
+}
+
+void UCEEffectorSubsystem::UpdateEffectorChannel()
+{
+	const UWorld* World = GetWorld();
+
+	if (!IsValid(World))
+	{
+		return;
+	}
+
+	if (EffectorsWeak.IsEmpty())
+	{
+		return;
+	}
+
+	// Reserve space in channel for each effectors
+	static const FNiagaraDataChannelSearchParameters SearchParameters;
+	UNiagaraDataChannelWriter* ChannelWriter = UNiagaraDataChannelLibrary::WriteToNiagaraDataChannel(World, EffectorDataChannelAsset.Get(), SearchParameters, EffectorsWeak.Num(), true, true, true, UCEEffectorSubsystem::StaticClass()->GetName());
+
+	if (!ChannelWriter)
+	{
+		UE_LOG(LogACEEffectorSubsystem, Warning, TEXT("Effector data channel writer is invalid"));
+		return;
+	}
+
+	// Remove invalid effectors and push updates to effector assigned channel indexes
+	int32 EffectorIndex = 0;
+	for (TArray<TWeakObjectPtr<ACEEffectorActor>>::TIterator It(EffectorsWeak); It; ++It)
+	{
+		ACEEffectorActor* Effector = It->Get();
+
+		if (!Effector)
+		{
+			It.RemoveCurrent();
+			continue;
+		}
+
+		FCEClonerEffectorChannelData& ChannelData = Effector->GetEffectorChannelData();
+
+		const bool bIdentifierChanged = ChannelData.Identifier != EffectorIndex;
+
+		// Set channel before writing
+		ChannelData.Identifier = EffectorIndex++;
+
+		// Push effector data to channel
+		ChannelData.Write(ChannelWriter);
+
+		// When changed, update cloners DI linked to this effector
+		if (bIdentifierChanged)
+		{
+			Effector->OnEffectorIdentifierChanged();
+		}
+	}
+}
+
+bool UCEEffectorSubsystem::IsTickableInEditor() const
+{
+	return true;
+}
+
+TStatId UCEEffectorSubsystem::GetStatId() const
+{
+	RETURN_QUICK_DECLARE_CYCLE_STAT(UCEEffectorSubsystem, STATGROUP_Tickables);
+}
+
+void UCEEffectorSubsystem::Tick(float InDeltaTime)
+{
+	UpdateEffectorChannel();
+}
diff --git a/Engine/Plugins/Experimental/ClonerEffector/Source/ClonerEffector/Private/Subsystems/CEEffectorSubsystem.h b/Engine/Plugins/Experimental/ClonerEffector/Source/ClonerEffector/Private/Subsystems/CEEffectorSubsystem.h
new file mode 100644
--- /dev/null
+++ b/Engine/Plugins/Experimental/ClonerEffector/Source/ClonerEffector/Private/Subsystems/CEEffectorSubsystem.h
@@ -0,0 +1,55 @@
+// Copyright Epic Games, Inc. All Rights Reserved.
+
+#pragma once
+
+#include "Engine/World.h"
+#include "Subsystems/WorldSubsystem.h"
+#include "UObject/WeakObjectPtrTemplates.h"
+#include "CEEffectorSubsystem.generated.h"
+
+class ACEEffectorActor;
+class UNiagaraDataChannelAsset;
+
+UCLASS()
+class UCEEffectorSubsystem : public UTickableWorldSubsystem
+{
+	GENERATED_BODY()
+
+public:
+	DECLARE_MULTICAST_DELEGATE_OneParam(FOnSubsystemInitialized, const UWorld*)
+	static FOnSubsystemInitialized OnSubsystemInitializedDelegate;
+
+	static inline constexpr TCHAR DataChannelAssetPath[] = TEXT("/Script/Niagara.NiagaraDataChannelAsset'/ClonerEffector/Channels/NDC_Effector.NDC_Effector'");
+
+	/** Get this subsystem instance */
+	static UCEEffectorSubsystem* Get(const UWorld* InWorld = GWorld);
+
+	//~ Begin USubsystem
+	virtual void Initialize(FSubsystemCollectionBase& Collection) override;
+	virtual void PostInitialize() override;
+	//~ End USubsystem
+
+	/** Registers an effector actor to use it within a effector channel */
+	bool RegisterChannelEffector(ACEEffectorActor* InEffector);
+
+	/** Unregister an effector actor used within a effector channel */
+	bool UnregisterChannelEffector(ACEEffectorActor* InEffector);
+
+protected:
+	//~ Begin FTickableGameObject interface
+	virtual bool IsTickableInEditor() const override;
+	virtual TStatId GetStatId() const override;
+	virtual void Tick(float InDeltaTime) override;
+	//~ End FTickableGameObject interface
+
+	/** Updates all registered effectors */
+	void UpdateEffectorChannel();
+
+	/** Ordered effectors included in this channel */
+	UPROPERTY()
+	TArray<TWeakObjectPtr<ACEEffectorActor>> EffectorsWeak;
+
+	/** This represents the data channel structure for effector */
+	UPROPERTY()
+	TObjectPtr<UNiagaraDataChannelAsset> EffectorDataChannelAsset;
+};
\ No newline at end of file
diff --git a/Engine/Plugins/Experimental/ClonerEffector/Source/ClonerEffector/Public/CEClonerEffectorShared.h b/Engine/Plugins/Experimental/ClonerEffector/Source/ClonerEffector/Public/CEClonerEffectorShared.h
new file mode 100644
--- /dev/null
+++ b/Engine/Plugins/Experimental/ClonerEffector/Source/ClonerEffector/Public/CEClonerEffectorShared.h
@@ -0,0 +1,665 @@
+// Copyright Epic Games, Inc. All Rights Reserved.
+
+#pragma once
+
+#include "Components/SphereComponent.h"
+#include "Engine/StaticMesh.h"
+#include "Materials/MaterialInterface.h"
+#include "NiagaraDataInterfaceArrayInt.h"
+#include "CEClonerEffectorShared.generated.h"
+
+class UNiagaraDataChannelReader;
+class UNiagaraDataChannelWriter;
+class AActor;
+class UDynamicMesh;
+class UNiagaraSystem;
+class USplineComponent;
+
+/** Enumerate the different layouts available to dispose cloner instances */
+UENUM(BlueprintType)
+enum class ECEClonerLayout : uint8
+{
+	Grid,
+	Line,
+	Circle,
+	Cylinder,
+	Sphere,
+	Honeycomb,
+	SampleMesh,
+	SampleSpline
+};
+
+/** Enumerate the axis available to use */
+UENUM(BlueprintType)
+enum class ECEClonerAxis : uint8
+{
+	X,
+	Y,
+	Z,
+	Custom
+};
+
+/** Enumerate the planes available to use */
+UENUM(BlueprintType)
+enum class ECEClonerPlane : uint8
+{
+	XY,
+	YZ,
+	XZ,
+	Custom
+};
+
+/** Enumerate the mesh render modes available to render a cloner instance */
+UENUM(BlueprintType)
+enum class ECEClonerMeshRenderMode : uint8
+{
+	/** Iterate through each attachment mesh available */
+	Iterate,
+	/** Pick randomly through each attachment mesh available, update cloner seed for variations */
+	Random,
+	/** Blend based on the total cloner instances and attachment mesh available */
+	Blend
+};
+
+/** Enumerate the grid constraints available when cloner layout selected is grid */
+UENUM(BlueprintType)
+enum class ECEClonerGridConstraint : uint8
+{
+	None,
+	Sphere,
+	Cylinder,
+	Texture
+};
+
+/* Enumerates all easing functions that are available to apply on weights
+See https://easings.net/ for curves visualizations and web open domain implementations
+Used as one enum to send the index to niagara as uint8 and apply easing directly in niagara */
+UENUM(BlueprintType)
+enum class ECEClonerEasing : uint8
+{
+	Linear,
+
+	InSine,
+	OutSine,
+	InOutSine,
+
+	InQuad,
+	OutQuad,
+	InOutQuad,
+
+	InCubic,
+	OutCubic,
+	InOutCubic,
+
+	InQuart,
+	OutQuart,
+	InOutQuart,
+
+	InQuint,
+	OutQuint,
+	InOutQuint,
+
+	InExpo,
+	OutExpo,
+	InOutExpo,
+
+	InCirc,
+	OutCirc,
+	InOutCirc,
+
+	InBack,
+	OutBack,
+	InOutBack,
+
+	InElastic,
+	OutElastic,
+	InOutElastic,
+
+	InBounce,
+	OutBounce,
+	InOutBounce,
+
+	Random
+};
+
+/** Enumerate the mesh asset to look for when mesh layout is selected */
+UENUM(BlueprintType)
+enum class ECEClonerMeshAsset : uint8
+{
+	StaticMesh,
+	SkeletalMesh
+};
+
+/** Enumerate the mesh sample dataset to pick from when mesh layout is selected */
+UENUM(BlueprintType)
+enum class ECEClonerMeshSampleData : uint8
+{
+	Vertices,
+	Triangles,
+	Sockets,
+	Bones,
+	Sections
+};
+
+/** Enumerates the effector shapes available */
+UENUM(BlueprintType)
+enum class ECEClonerEffectorType : uint8
+{
+	/** Clones inside the sphere radius will be affected by the effector */
+	Sphere,
+	/** Clones between two planes will be affected by the effector */
+	Plane,
+	/** Clones inside the box extent will be affected by the effector */
+	Box,
+	/** All clones will be affected by the effector with the same max weight */
+	Unbound
+};
+
+/** Enumerates the effector mode available */
+UENUM(BlueprintType)
+enum class ECEClonerEffectorMode : uint8
+{
+	/** Control clones offset, rotation, scale manually */
+	Default,
+	/** Rotates clones towards a target actor */
+	Target,
+	/** Randomly applies curl noise across the field zone */
+	NoiseField
+};
+
+UENUM()
+enum class ECEClonerAttachmentStatus : uint8
+{
+	/** Item should be removed, no longer valid */
+	Invalid,
+	/** Item should be updated, changes detected */
+	Outdated,
+	/** Item is up to date, no changes needed */
+	Updated,
+	/** Item is being updated at the moment */
+	Updating
+};
+
+/** Enumerate all texture channels to sample for constraint */
+UENUM()
+enum class ECEClonerTextureSampleChannel : uint8
+{
+	RGBLuminance,
+	RGBAverage,
+	RGBMax,
+	R,
+	G,
+	B,
+	A
+};
+
+/** Enumerate all operation compare mode for constraint */
+UENUM()
+enum class ECEClonerCompareMode : uint8
+{
+	Greater,
+	GreaterEqual,
+	Equal,
+	NotEqual,
+	Less,
+	LessEqual
+};
+
+/** Enumerates all modes to handle cloner spawn mode */
+UENUM()
+enum class ECEClonerSpawnLoopMode : uint8
+{
+	/** Cloner spawns once and then enters idle mode */
+	Once,
+	/** Cloner spawns multiple times and then enters idle mode */
+	Multiple,
+	/** Cloner spawns infinitely and never enters idle mode */
+	Infinite
+};
+
+/** Enumerates all modes for how clones are spawned */
+UENUM()
+enum class ECEClonerSpawnBehaviorMode : uint8
+{
+	/** Spawns instantly the number of clones needed for the layout */
+	Instant,
+	/** Spawns at a specific rate per second during the spawn loop duration */
+	Rate,
+};
+
+USTRUCT()
+struct FCEClonerAttachmentItem
+{
+	GENERATED_BODY()
+
+	/** Children attached to this actor, order is not important here as they are combined into one */
+	UPROPERTY()
+	TSet<TWeakObjectPtr<AActor>> ChildrenActors;
+
+	/** Current cloner attached actor represented by this item */
+	UPROPERTY()
+	TWeakObjectPtr<AActor> ItemActor;
+
+	/** Parent of this item actor or null if below cloner */
+	UPROPERTY()
+	TWeakObjectPtr<AActor> ParentActor;
+
+	/** Actual baked mesh for the current actor item,
+		dynamic because it's easier for merging them together and avoid conversions again */
+	UPROPERTY()
+	TObjectPtr<UDynamicMesh> BakedMesh;
+
+	/** Actual baked materials for the current actor item */
+	UPROPERTY()
+	TArray<TWeakObjectPtr<UMaterialInterface>> BakedMaterials;
+
+	/** Status of the baked mesh, does it needs to be updated */
+	UPROPERTY()
+	ECEClonerAttachmentStatus MeshStatus = ECEClonerAttachmentStatus::Outdated;
+
+	/** Last actor item transform, used to trigger an update if changed */
+	UPROPERTY()
+	FTransform ActorTransform;
+
+	/** Item is root cloner actor */
+	UPROPERTY()
+	bool bRootItem = false;
+
+	/** Status of this attachment item */
+	UPROPERTY()
+	ECEClonerAttachmentStatus Status = ECEClonerAttachmentStatus::Outdated;
+
+	friend uint32 GetTypeHash(const FCEClonerAttachmentItem& InItem)
+	{
+		return GetTypeHash(InItem.ItemActor);
+	}
+
+	bool operator==(const FCEClonerAttachmentItem& InOther) const
+	{
+		return ItemActor.Get() == InOther.ItemActor.Get();
+	}
+};
+
+USTRUCT()
+struct FCEClonerAttachmentTree
+{
+	GENERATED_BODY()
+
+	/** All cloner attached actor items */
+	UPROPERTY()
+	TMap<TWeakObjectPtr<AActor>, FCEClonerAttachmentItem> ItemAttachmentMap;
+
+	/** Actors directly attached to the cloner actor, order is important here */
+	UPROPERTY()
+	TArray<TWeakObjectPtr<AActor>> RootActors;
+
+	/** Merged static meshes corresponding to the root actors for niagara,
+		result of merging dynamic meshes together */
+	UPROPERTY()
+	TArray<TObjectPtr<UStaticMesh>> MergedBakedMeshes;
+
+	/** Status of the cloner tree */
+	UPROPERTY()
+	ECEClonerAttachmentStatus Status = ECEClonerAttachmentStatus::Updated;
+
+	void Reset()
+	{
+		ItemAttachmentMap.Empty();
+		RootActors.Empty();
+		MergedBakedMeshes.Empty();
+		Status = ECEClonerAttachmentStatus::Updated;
+	}
+};
+
+struct FCEClonerEffectorChannelData
+{
+	friend class UCEEffectorSubsystem;
+
+	/** General */
+	static constexpr const TCHAR* EasingName = TEXT("Easing");
+	static constexpr const TCHAR* ModeName = TEXT("Mode");
+	static constexpr const TCHAR* TypeName = TEXT("Type");
+	static constexpr const TCHAR* MagnitudeName = TEXT("Magnitude");
+	static constexpr const TCHAR* InnerExtentName = TEXT("InnerExtent");
+	static constexpr const TCHAR* OuterExtentName = TEXT("OuterExtent");
+	static constexpr const TCHAR* LocationDeltaName = TEXT("LocationDelta");
+	static constexpr const TCHAR* RotationDeltaName = TEXT("RotationDelta");
+	static constexpr const TCHAR* ScaleDeltaName = TEXT("ScaleDelta");
+	static constexpr const TCHAR* LocationName = TEXT("Location");
+	static constexpr const TCHAR* RotationName = TEXT("Rotation");
+	static constexpr const TCHAR* ScaleName = TEXT("Scale");
+	static constexpr const TCHAR* FrequencyName = TEXT("Frequency");
+	static constexpr const TCHAR* PanName = TEXT("Pan");
+
+	/** Forces */
+	static constexpr const TCHAR* OrientationForceRateName = TEXT("OrientationForceRate");
+	static constexpr const TCHAR* OrientationForceMinName = TEXT("OrientationForceMin");
+	static constexpr const TCHAR* OrientationForceMaxName = TEXT("OrientationForceMax");
+	static constexpr const TCHAR* VortexForceAmountName = TEXT("VortexForceAmount");
+	static constexpr const TCHAR* VortexForceAxisName = TEXT("VortexForceAxis");
+	static constexpr const TCHAR* CurlNoiseForceStrengthName = TEXT("CurlNoiseForceStrength");
+	static constexpr const TCHAR* CurlNoiseForceFrequencyName = TEXT("CurlNoiseForceFrequency");
+	static constexpr const TCHAR* AttractionForceStrengthName = TEXT("AttractionForceStrength");
+	static constexpr const TCHAR* AttractionForceFalloffName = TEXT("AttractionForceFalloff");
+	static constexpr const TCHAR* GravityForceAccelerationName = TEXT("GravityForceAcceleration");
+
+	int32 GetIdentifier() const
+	{
+		return Identifier;
+	}
+
+	/** General parameters */
+	ECEClonerEasing Easing = ECEClonerEasing::Linear;
+	ECEClonerEffectorMode Mode = ECEClonerEffectorMode::Default;
+	ECEClonerEffectorType Type = ECEClonerEffectorType::Sphere;
+	float Magnitude = 0.f;
+	FVector InnerExtent = FVector::ZeroVector;
+	FVector OuterExtent = FVector::ZeroVector;
+	FVector LocationDelta = FVector::ZeroVector;
+	FQuat RotationDelta = FQuat::Identity;
+	FVector ScaleDelta = FVector::OneVector;
+	FVector Location = FVector::ZeroVector;
+	FQuat Rotation = FQuat::Identity;
+	FVector Scale = FVector::OneVector;
+	float Frequency = 1.f;
+	FVector Pan = FVector::ZeroVector;
+
+	/** Forces parameters */
+	float OrientationForceRate = 0.f;
+	FVector OrientationForceMin = FVector::ZeroVector;
+	FVector OrientationForceMax = FVector::ZeroVector;
+	float VortexForceAmount = 0.f;
+	FVector VortexForceAxis = FVector::ZeroVector;
+	float CurlNoiseForceStrength = 0.f;
+	float CurlNoiseForceFrequency = 0.f;
+	float AttractionForceStrength = 0.f;
+	float AttractionForceFalloff = 0.f;
+	FVector GravityForceAcceleration = FVector::ZeroVector;
+
+protected:
+	/** Cache effector identifier to detect a change and update cloners DI */
+	int32 Identifier = INDEX_NONE;
+
+	void Write(UNiagaraDataChannelWriter* InWriter) const;
+	void Read(const UNiagaraDataChannelReader* InReader);
+};
+
+USTRUCT()
+struct FCEClonerEffectorDataInterfaces
+{
+	friend class UCEClonerLayoutBase;
+
+	GENERATED_BODY()
+
+	static inline const FName IndexName = TEXT("EffectorIndexArray");
+
+	explicit FCEClonerEffectorDataInterfaces(const UNiagaraSystem* InSystem);
+	FCEClonerEffectorDataInterfaces() = default;
+
+	void CopyTo(FCEClonerEffectorDataInterfaces& InOther) const;
+	void Resize(int32 InSize) const;
+	void Remove(int32 InIndex) const;
+	bool IsValid() const;
+	int32 Num() const;
+
+	UNiagaraDataInterfaceArrayInt32* GetIndexArray() const;
+
+protected:
+	UPROPERTY()
+	TMap<FName, TObjectPtr<UNiagaraDataInterface>> DataInterfaces;
+};
+
+USTRUCT(BlueprintType)
+struct CLONEREFFECTOR_API FCEClonerSampleMeshOptions
+{
+	GENERATED_BODY()
+
+	UPROPERTY(BlueprintReadWrite, EditInstanceOnly, Category="Cloner")
+	int32 Count = 3 * 3 * 3;
+
+	UPROPERTY(BlueprintReadWrite, EditInstanceOnly, Category="Cloner")
+	ECEClonerMeshAsset Asset = ECEClonerMeshAsset::StaticMesh;
+
+	UPROPERTY(BlueprintReadWrite, EditInstanceOnly, Category="Cloner")
+	ECEClonerMeshSampleData SampleData = ECEClonerMeshSampleData::Vertices;
+
+	UPROPERTY(BlueprintReadWrite, EditInstanceOnly, Category="Cloner")
+	TWeakObjectPtr<AActor> SampleActor = nullptr;
+
+	UPROPERTY()
+	TWeakObjectPtr<USceneComponent> SceneComponent = nullptr;
+};
+
+USTRUCT(BlueprintType)
+struct CLONEREFFECTOR_API FCEClonerSampleSplineOptions
+{
+	GENERATED_BODY()
+
+	UPROPERTY(BlueprintReadWrite, EditInstanceOnly, Category="Cloner")
+	int32 Count = 3 * 3;
+
+	UPROPERTY(BlueprintReadWrite, EditInstanceOnly, Category="Cloner")
+	TWeakObjectPtr<AActor> SplineActor = nullptr;
+
+	UPROPERTY()
+	TWeakObjectPtr<USplineComponent> SplineComponent = nullptr;
+};
+
+USTRUCT(BlueprintType)
+struct CLONEREFFECTOR_API FCEClonerGridConstraintSphere
+{
+	GENERATED_BODY()
+
+	UPROPERTY(BlueprintReadWrite, EditInstanceOnly, Category="Cloner")
+	float Radius = 400.f;
+
+	UPROPERTY(BlueprintReadWrite, EditInstanceOnly, Category="Cloner")
+	FVector Center = FVector(0);
+};
+
+USTRUCT(BlueprintType)
+struct CLONEREFFECTOR_API FCEClonerGridConstraintCylinder
+{
+	GENERATED_BODY()
+
+	UPROPERTY(BlueprintReadWrite, EditInstanceOnly, Category="Cloner")
+	float Radius = 400.f;
+
+	UPROPERTY(BlueprintReadWrite, EditInstanceOnly, Category="Cloner")
+	float Height = 800.f;
+
+	UPROPERTY(BlueprintReadWrite, EditInstanceOnly, Category="Cloner")
+	FVector Center = FVector(0);
+};
+
+USTRUCT(BlueprintType)
+struct CLONEREFFECTOR_API FCEClonerGridConstraintTexture
+{
+	GENERATED_BODY()
+
+	UPROPERTY(BlueprintReadWrite, EditInstanceOnly, Category="Cloner")
+	TWeakObjectPtr<UTexture> Texture;
+
+	UPROPERTY(BlueprintReadWrite, EditInstanceOnly, Category="Cloner")
+	ECEClonerTextureSampleChannel Channel = ECEClonerTextureSampleChannel::RGBLuminance;
+
+	UPROPERTY(BlueprintReadWrite, EditInstanceOnly, Category="Cloner")
+	ECEClonerPlane Plane = ECEClonerPlane::XY;
+
+	UPROPERTY(BlueprintReadWrite, EditInstanceOnly, Category="Cloner")
+	ECEClonerCompareMode CompareMode = ECEClonerCompareMode::Greater;
+
+	UPROPERTY(BlueprintReadWrite, EditInstanceOnly, Category="Cloner")
+	float Threshold = 0.f;
+};
+
+USTRUCT(BlueprintType)
+struct CLONEREFFECTOR_API FCEClonerGridLayoutOptions
+{
+	GENERATED_BODY()
+
+	UPROPERTY(BlueprintReadWrite, EditInstanceOnly, Category="Cloner")
+	FIntVector Count = FIntVector(3, 3, 3);
+
+	UPROPERTY(BlueprintReadWrite, EditInstanceOnly, Category="Cloner")
+	FVector Spacing = FVector(400.f, 400.f, 400.f);
+
+	UPROPERTY(BlueprintReadWrite, EditInstanceOnly, Category="Cloner")
+	ECEClonerGridConstraint Constraint = ECEClonerGridConstraint::None;
+
+	UPROPERTY(BlueprintReadWrite, EditInstanceOnly, Category="Cloner", meta=(EditCondition="Constraint != ECEClonerGridConstraint::None", EditConditionHides))
+	bool bInvertConstraint = false;
+
+	UPROPERTY(BlueprintReadWrite, EditInstanceOnly, Category="Cloner", meta=(EditCondition="Constraint == ECEClonerGridConstraint::Sphere", EditConditionHides))
+	FCEClonerGridConstraintSphere SphereConstraint;
+
+	UPROPERTY(BlueprintReadWrite, EditInstanceOnly, Category="Cloner", meta=(EditCondition="Constraint == ECEClonerGridConstraint::Cylinder", EditConditionHides))
+	FCEClonerGridConstraintCylinder CylinderConstraint;
+
+	UPROPERTY(BlueprintReadWrite, EditInstanceOnly, Category="Cloner", meta=(EditCondition="Constraint == ECEClonerGridConstraint::Texture", EditConditionHides))
+	FCEClonerGridConstraintTexture TextureConstraint;
+};
+
+USTRUCT(BlueprintType)
+struct CLONEREFFECTOR_API FCEClonerLineLayoutOptions
+{
+	GENERATED_BODY()
+
+	UPROPERTY(BlueprintReadWrite, EditInstanceOnly, Category="Cloner")
+	int32 Count = 3 * 3;
+
+	UPROPERTY(BlueprintReadWrite, EditInstanceOnly, Category="Cloner")
+	float Spacing = 400.f;
+
+	UPROPERTY(BlueprintReadWrite, EditInstanceOnly, Category="Cloner")
+	ECEClonerAxis Axis = ECEClonerAxis::Y;
+
+	UPROPERTY(BlueprintReadWrite, EditInstanceOnly, Category="Cloner", meta=(ClampMin="0", ClampMax="1", EditCondition="Axis == ECEClonerAxis::Custom", EditConditionHides))
+	FVector Direction = FVector::YAxisVector;
+
+	UPROPERTY(BlueprintReadWrite, EditInstanceOnly, Category="Cloner")
+	FRotator Rotation = FRotator(0.f);
+};
+
+USTRUCT(BlueprintType)
+struct CLONEREFFECTOR_API FCEClonerCircleLayoutOptions
+{
+	GENERATED_BODY()
+
+	UPROPERTY(BlueprintReadWrite, EditInstanceOnly, Category="Cloner")
+	int32 Count = 3 * 3;
+
+	UPROPERTY(BlueprintReadWrite, EditInstanceOnly, Category="Cloner")
+	float Radius = 400.f;
+
+	UPROPERTY(BlueprintReadWrite, EditInstanceOnly, Category="Cloner")
+	float AngleStart = 0.f;
+
+	UPROPERTY(BlueprintReadWrite, EditInstanceOnly, Category="Cloner")
+	float AngleRatio = 1.f;
+
+	UPROPERTY(BlueprintReadWrite, EditInstanceOnly, Category="Cloner")
+	bool bOrientMesh = false;
+
+	UPROPERTY(BlueprintReadWrite, EditInstanceOnly, Category="Cloner")
+	ECEClonerPlane Plane = ECEClonerPlane::XY;
+
+	UPROPERTY(BlueprintReadWrite, EditInstanceOnly, Category="Cloner", meta=(EditCondition="Plane == ECEClonerPlane::Custom", EditConditionHides))
+	FRotator Rotation = FRotator(0.f, 0.f, 0.f);
+
+	UPROPERTY(BlueprintReadWrite, EditInstanceOnly, Category="Cloner", meta=(ClampMin="0", AllowPreserveRatio, Delta="0.01"))
+	FVector Scale = FVector(1.f, 1.f, 1.f);
+};
+
+USTRUCT(BlueprintType)
+struct CLONEREFFECTOR_API FCEClonerCylinderLayoutOptions
+{
+	GENERATED_BODY()
+
+	UPROPERTY(BlueprintReadWrite, EditInstanceOnly, Category="Cloner")
+	int32 BaseCount = 3 * 3;
+
+	UPROPERTY(BlueprintReadWrite, EditInstanceOnly, Category="Cloner")
+	int32 HeightCount = 3;
+
+	UPROPERTY(BlueprintReadWrite, EditInstanceOnly, Category="Cloner")
+	float Height = 400.f;
+
+	UPROPERTY(BlueprintReadWrite, EditInstanceOnly, Category="Cloner")
+	float Radius = 400.f;
+
+	UPROPERTY(BlueprintReadWrite, EditInstanceOnly, Category="Cloner")
+	float AngleStart = 0.f;
+
+	UPROPERTY(BlueprintReadWrite, EditInstanceOnly, Category="Cloner")
+	float AngleRatio = 1.f;
+
+	UPROPERTY(BlueprintReadWrite, EditInstanceOnly, Category="Cloner")
+	bool bOrientMesh = false;
+
+	UPROPERTY(BlueprintReadWrite, EditInstanceOnly, Category="Cloner")
+	ECEClonerPlane Plane = ECEClonerPlane::XY;
+
+	UPROPERTY(BlueprintReadWrite, EditInstanceOnly, Category="Cloner", meta=(EditCondition="Plane == ECEClonerPlane::Custom", EditConditionHides))
+	FRotator Rotation = FRotator(0.f, 0.f, 0.f);
+
+	UPROPERTY(BlueprintReadWrite, EditInstanceOnly, Category="Cloner", meta=(ClampMin="0", AllowPreserveRatio, Delta="0.01"))
+	FVector Scale = FVector(1.f, 1.f, 1.f);
+};
+
+USTRUCT(BlueprintType)
+struct CLONEREFFECTOR_API FCEClonerSphereLayoutOptions
+{
+	GENERATED_BODY()
+
+	UPROPERTY(BlueprintReadWrite, EditInstanceOnly, Category="Cloner")
+	int32 Count = 3 * 3 * 3;
+
+	UPROPERTY(BlueprintReadWrite, EditInstanceOnly, Category="Cloner")
+	float Radius = 400.f;
+
+	UPROPERTY(BlueprintReadWrite, EditInstanceOnly, Category="Cloner", meta=(ClampMin="0", ClampMax="1"))
+	float Ratio = 1.f;
+
+	UPROPERTY(BlueprintReadWrite, EditInstanceOnly, Category="Cloner")
+	bool bOrientMesh = false;
+
+	UPROPERTY(BlueprintReadWrite, EditInstanceOnly, Category="Cloner")
+	FRotator Rotation = FRotator(0.f, 0.f, 0.f);
+
+	UPROPERTY(BlueprintReadWrite, EditInstanceOnly, Category="Cloner", meta=(ClampMin="0", AllowPreserveRatio, Delta="0.01"))
+	FVector Scale = FVector(1.f, 1.f, 1.f);
+};
+
+USTRUCT(BlueprintType)
+struct CLONEREFFECTOR_API FCEClonerHoneycombLayoutOptions
+{
+	GENERATED_BODY()
+
+	UPROPERTY(BlueprintReadWrite, EditInstanceOnly, Category="Cloner")
+	ECEClonerPlane Plane = ECEClonerPlane::XY;
+
+	UPROPERTY(BlueprintReadWrite, EditInstanceOnly, Category="Cloner")
+	int32 WidthCount = 3;
+
+	UPROPERTY(BlueprintReadWrite, EditInstanceOnly, Category="Cloner")
+	int32 HeightCount = 3;
+
+	UPROPERTY(BlueprintReadWrite, EditInstanceOnly, Category="Cloner")
+	float WidthOffset = 0.f;
+
+	UPROPERTY(BlueprintReadWrite, EditInstanceOnly, Category="Cloner")
+	float HeightOffset = 0.5f;
+
+	UPROPERTY(BlueprintReadWrite, EditInstanceOnly, Category="Cloner")
+	float WidthSpacing = 400.f;
+
+	UPROPERTY(BlueprintReadWrite, EditInstanceOnly, Category="Cloner")
+	float HeightSpacing = 400.f;
+};
+
+namespace UE::ClonerEffector
+{
+	void SetBillboardComponentSprite(const AActor* InActor, const FString& InTexturePath);
+	void SetBillboardComponentVisibility(const AActor* InActor, bool bInVisibility);
+}
\ No newline at end of file
