MotionDesign

diff --git a/Engine/Plugins/Experimental/ClonerEffector/Source/ClonerEffector/Public/Cloner/CEClonerActor.h b/Engine/Plugins/Experimental/ClonerEffector/Source/ClonerEffector/Public/Cloner/CEClonerActor.h
new file mode 100644
--- /dev/null
+++ b/Engine/Plugins/Experimental/ClonerEffector/Source/ClonerEffector/Public/Cloner/CEClonerActor.h
@@ -0,0 +1,701 @@
+// Copyright Epic Games, Inc. All Rights Reserved.
+
+#pragma once
+
+#include "CEClonerEffectorShared.h"
+#include "Curves/RichCurve.h"
+#include "GameFramework/Actor.h"
+#include "Layouts/CEClonerLayoutBase.h"
+#include "NiagaraMeshRendererProperties.h"
+#include "CEPropertyChangeDispatcher.h"
+#include "CEClonerActor.generated.h"
+
+class AActor;
+class ACEEffectorActor;
+class UCEClonerComponent;
+class UCEClonerLayoutBase;
+class UNiagaraDataInterfaceCurve;
+class USceneComponent;
+
+UCLASS(MinimalAPI, BlueprintType, HideCategories=(Rendering,Replication,Collision,HLOD,Physics,Networking,Input,Actor,Cooking,LevelInstance,Streaming), DisplayName = "Motion Design Cloner Actor")
+class ACEClonerActor : public AActor
+{
+	GENERATED_BODY()
+
+	friend class ACEEffectorActor;
+	friend class UCEClonerLayoutBase;
+
+	friend class FAvaClonerActorVisualizer;
+	friend class FCEEditorClonerDetailCustomization;
+
+public:
+	static inline const FString DefaultLabel = TEXT("Cloner");
+
+	ACEClonerActor();
+
+	UFUNCTION(BlueprintCallable, Category="Cloner")
+	CLONEREFFECTOR_API void SetTreeUpdateInterval(float InInterval);
+
+	UFUNCTION(BlueprintPure, Category="Cloner")
+	float GetTreeUpdateInterval() const
+	{
+		return TreeUpdateInterval;
+	}
+
+	UFUNCTION(BlueprintCallable, Category="Cloner")
+	CLONEREFFECTOR_API void SetMeshRenderMode(ECEClonerMeshRenderMode InMode);
+
+	UFUNCTION(BlueprintPure, Category="Cloner")
+	ECEClonerMeshRenderMode GetMeshRenderMode() const
+	{
+		return MeshRenderMode;
+	}
+
+	UFUNCTION(BlueprintCallable, Category="Cloner")
+	CLONEREFFECTOR_API void SetMeshFacingMode(ENiagaraMeshFacingMode InMode);
+
+	UFUNCTION(BlueprintPure, Category="Cloner")
+	ENiagaraMeshFacingMode GetMeshFacingMode() const
+	{
+		return MeshFacingMode;
+	}
+
+	UFUNCTION(BlueprintCallable, Category="Cloner")
+	CLONEREFFECTOR_API void SetMeshCastShadows(bool InbCastShadows);
+
+	UFUNCTION(BlueprintPure, Category="Cloner")
+	bool GetMeshCastShadows() const
+	{
+		return bMeshCastShadows;
+	}
+
+	CLONEREFFECTOR_API void SetDefaultMeshes(const TArray<TObjectPtr<UStaticMesh>>& InMeshes);
+
+	const TArray<TObjectPtr<UStaticMesh>>& GetDefaultMeshes() const
+	{
+		return DefaultMeshes;
+	}
+
+	UFUNCTION(BlueprintCallable, Category="Cloner", meta=(DisplayName="SetDefaultMeshes"))
+	CLONEREFFECTOR_API void BP_SetDefaultMeshes(const TArray<UStaticMesh*>& InMeshes);
+
+	UFUNCTION(BlueprintPure, Category="Cloner", meta=(DisplayName="GetDefaultMeshes"))
+	CLONEREFFECTOR_API TArray<UStaticMesh*> BP_GetDefaultMeshes() const;
+
+	UFUNCTION(BlueprintCallable, Category="Cloner")
+	CLONEREFFECTOR_API void SetSeed(int32 InSeed);
+
+	UFUNCTION(BlueprintPure, Category="Cloner")
+	int32 GetSeed() const
+	{
+		return Seed;
+	}
+
+	UFUNCTION(BlueprintCallable, Category="Cloner")
+	CLONEREFFECTOR_API void SetDeltaStepEnabled(bool bInEnabled);
+
+	UFUNCTION(BlueprintPure, Category="Cloner")
+	bool GetDeltaStepEnabled() const
+	{
+		return bDeltaStepEnabled;
+	}
+
+	UFUNCTION(BlueprintCallable, Category="Cloner")
+	CLONEREFFECTOR_API void SetDeltaStepRotation(const FRotator& InRotation);
+
+	UFUNCTION(BlueprintPure, Category="Cloner")
+	FRotator GetDeltaStepRotation() const
+	{
+		return DeltaStepRotation;
+	}
+
+	UFUNCTION(BlueprintCallable, Category="Cloner")
+	CLONEREFFECTOR_API void SetDeltaStepScale(const FVector& InScale);
+
+	UFUNCTION(BlueprintPure, Category="Cloner")
+	FVector GetDeltaStepScale() const
+	{
+		return DeltaStepScale;
+	}
+
+	UFUNCTION(BlueprintCallable, Category="Cloner")
+	CLONEREFFECTOR_API void SetEnabled(bool bInEnable);
+
+	UFUNCTION(BlueprintPure, Category="Cloner")
+	bool GetEnabled() const
+	{
+		return bEnabled;
+	}
+
+	UFUNCTION(BlueprintCallable, Category="Cloner")
+	CLONEREFFECTOR_API void SetInvertProgress(bool bInInvertProgress);
+
+	UFUNCTION(BlueprintPure, Category="Cloner")
+	bool GetInvertProgress() const
+	{
+		return bInvertProgress;
+	}
+
+	UFUNCTION(BlueprintCallable, Category="Cloner")
+	CLONEREFFECTOR_API void SetProgress(float InProgress);
+
+	UFUNCTION(BlueprintPure, Category="Cloner")
+	float GetProgress() const
+	{
+		return Progress;
+	}
+
+	UFUNCTION(BlueprintCallable, Category="Cloner")
+	CLONEREFFECTOR_API void SetLayoutName(FName InLayoutName);
+
+	UFUNCTION(BlueprintPure, Category="Cloner")
+	FName GetLayoutName() const
+	{
+		return LayoutName;
+	}
+
+	UFUNCTION(BlueprintCallable, Category="Cloner")
+	CLONEREFFECTOR_API void SetRangeEnabled(bool bInRangeEnabled);
+
+	UFUNCTION(BlueprintPure, Category="Cloner")
+	bool GetRangeEnabled() const
+	{
+		return bRangeEnabled;
+	}
+
+	UFUNCTION(BlueprintCallable, Category="Cloner")
+	CLONEREFFECTOR_API void SetRangeOffsetMin(const FVector& InRangeOffsetMin);
+
+	UFUNCTION(BlueprintPure, Category="Cloner")
+	const FVector& GetRangeOffsetMin() const
+	{
+		return RangeOffsetMin;
+	}
+
+	UFUNCTION(BlueprintCallable, Category="Cloner")
+	CLONEREFFECTOR_API void SetRangeOffsetMax(const FVector& InRangeOffsetMax);
+
+	UFUNCTION(BlueprintPure, Category="Cloner")
+	const FVector& GetRangeOffsetMax() const
+	{
+		return RangeOffsetMax;
+	}
+
+	UFUNCTION(BlueprintCallable, Category="Cloner")
+	CLONEREFFECTOR_API void SetRangeRotationMin(const FRotator& InRangeRotationMin);
+
+	UFUNCTION(BlueprintPure, Category="Cloner")
+	const FRotator& GetRangeRotationMin() const
+	{
+		return RangeRotationMin;
+	}
+
+	UFUNCTION(BlueprintCallable, Category="Cloner")
+	CLONEREFFECTOR_API void SetRangeRotationMax(const FRotator& InRangeRotationMax);
+
+	UFUNCTION(BlueprintPure, Category="Cloner")
+	const FRotator& GetRangeRotationMax() const
+	{
+		return RangeRotationMax;
+	}
+
+	UFUNCTION(BlueprintCallable, Category="Cloner")
+	CLONEREFFECTOR_API void SetRangeScaleUniform(bool bInRangeScaleUniform);
+
+	UFUNCTION(BlueprintPure, Category="Cloner")
+	bool GetRangeScaleUniform() const
+	{
+		return bRangeScaleUniform;
+	}
+
+	UFUNCTION(BlueprintCallable, Category="Cloner")
+	CLONEREFFECTOR_API void SetRangeScaleMin(const FVector& InRangeScaleMin);
+
+	UFUNCTION(BlueprintPure, Category="Cloner")
+	const FVector& GetRangeScaleMin() const
+	{
+		return RangeScaleMin;
+	}
+
+	UFUNCTION(BlueprintCallable, Category="Cloner")
+	CLONEREFFECTOR_API void SetRangeScaleMax(const FVector& InRangeScaleMax);
+
+	UFUNCTION(BlueprintPure, Category="Cloner")
+	const FVector& GetRangeScaleMax() const
+	{
+		return RangeScaleMax;
+	}
+
+	UFUNCTION(BlueprintCallable, Category="Cloner")
+	CLONEREFFECTOR_API void SetRangeScaleUniformMin(float InRangeScaleUniformMin);
+
+	UFUNCTION(BlueprintPure, Category="Cloner")
+	float GetRangeScaleUniformMin() const
+	{
+		return RangeScaleUniformMin;
+	}
+
+	UFUNCTION(BlueprintCallable, Category="Cloner")
+	CLONEREFFECTOR_API void SetRangeScaleUniformMax(float InRangeScaleUniformMax);
+
+	UFUNCTION(BlueprintPure, Category="Cloner")
+	float GetRangeScaleUniformMax() const
+	{
+		return RangeScaleUniformMax;
+	}
+
+	UFUNCTION(BlueprintCallable, Category="Cloner")
+	CLONEREFFECTOR_API void SetSpawnLoopMode(ECEClonerSpawnLoopMode InMode);
+
+	UFUNCTION(BlueprintPure, Category="Cloner")
+	ECEClonerSpawnLoopMode GetSpawnLoopMode() const
+	{
+		return SpawnLoopMode;
+	}
+
+	UFUNCTION(BlueprintCallable, Category="Cloner")
+	CLONEREFFECTOR_API void SetSpawnLoopIterations(int32 InIterations);
+
+	UFUNCTION(BlueprintPure, Category="Cloner")
+	int32 GetSpawnLoopIterations() const
+	{
+		return SpawnLoopIterations;
+	}
+
+	UFUNCTION(BlueprintCallable, Category="Cloner")
+	CLONEREFFECTOR_API void SetSpawnLoopInterval(float InInterval);
+
+	UFUNCTION(BlueprintPure, Category="Cloner")
+	float GetSpawnLoopInterval() const
+	{
+		return SpawnLoopInterval;
+	}
+
+	UFUNCTION(BlueprintCallable, Category="Cloner")
+	CLONEREFFECTOR_API void SetSpawnBehaviorMode(ECEClonerSpawnBehaviorMode InMode);
+
+	UFUNCTION(BlueprintPure, Category="Cloner")
+	ECEClonerSpawnBehaviorMode GetSpawnBehaviorMode() const
+	{
+		return SpawnBehaviorMode;
+	}
+
+	UFUNCTION(BlueprintCallable, Category="Cloner")
+	CLONEREFFECTOR_API void SetSpawnRate(float InRate);
+
+	UFUNCTION(BlueprintPure, Category="Cloner")
+	float GetSpawnRate() const
+	{
+		return SpawnRate;
+	}
+
+	UFUNCTION(BlueprintCallable, Category="Cloner")
+	CLONEREFFECTOR_API void SetLifetimeEnabled(bool bInEnabled);
+
+	UFUNCTION(BlueprintPure, Category="Cloner")
+	bool GetLifetimeEnabled() const
+	{
+		return bLifetimeEnabled;
+	}
+
+	UFUNCTION(BlueprintCallable, Category="Cloner")
+	CLONEREFFECTOR_API void SetLifetimeMin(float InMin);
+
+	UFUNCTION(BlueprintPure, Category="Cloner")
+	float GetLifetimeMin() const
+	{
+		return LifetimeMin;
+	}
+
+	UFUNCTION(BlueprintCallable, Category="Cloner")
+	CLONEREFFECTOR_API void SetLifetimeMax(float InMax);
+
+	UFUNCTION(BlueprintPure, Category="Cloner")
+	float GetLifetimeMax() const
+	{
+		return LifetimeMax;
+	}
+
+	UFUNCTION(BlueprintCallable, Category="Cloner")
+	CLONEREFFECTOR_API void SetLifetimeScaleEnabled(bool bInEnabled);
+
+	UFUNCTION(BlueprintPure, Category="Cloner")
+	bool GetLifetimeScaleEnabled() const
+	{
+		return bLifetimeScaleEnabled;
+	}
+
+	UFUNCTION()
+	CLONEREFFECTOR_API void SetLifetimeScaleCurve(const FRichCurve& InCurve);
+
+	UFUNCTION()
+	const FRichCurve& GetLifetimeScaleCurve() const
+	{
+		return LifetimeScaleCurve;
+	}
+
+	UFUNCTION(BlueprintPure, Category="Cloner")
+	CLONEREFFECTOR_API UCEClonerLayoutBase* GetActiveLayout() const;
+
+	template<
+		typename InLayoutClass
+		UE_REQUIRES(TIsDerivedFrom<InLayoutClass, UCEClonerLayoutBase>::Value)>
+	bool IsActiveLayout() const
+	{
+		if (const UCEClonerLayoutBase* CurrentLayout = GetActiveLayout())
+		{
+			return CurrentLayout->GetClass() == InLayoutClass::StaticClass();
+		}
+
+		return false;
+	}
+
+	template<
+		typename InLayoutClass
+		UE_REQUIRES(TIsDerivedFrom<InLayoutClass, UCEClonerLayoutBase>::Value)>
+	InLayoutClass* GetActiveLayout() const
+	{
+		return Cast<InLayoutClass>(GetActiveLayout());
+	}
+
+	/** Returns the number of meshes this cloner currently handles */
+	UFUNCTION(BlueprintPure, Category="Cloner")
+	CLONEREFFECTOR_API int32 GetMeshCount() const;
+
+	/** Gets the number of effectors applied on this cloner */
+	UFUNCTION(BlueprintPure, Category="Cloner")
+	CLONEREFFECTOR_API int32 GetEffectorCount() const;
+
+#if WITH_EDITOR
+	/** This will force an update of the cloner attachment tree */
+	UFUNCTION(CallInEditor, Category="Cloner")
+	CLONEREFFECTOR_API void ForceUpdateCloner();
+
+	/** This will spawn an effector actor and link it to this cloner */
+	UFUNCTION(CallInEditor, Category="Cloner")
+	CLONEREFFECTOR_API void SpawnLinkedEffector();
+
+	/** When the cloner actor is created, we spawn a default actor attached */
+	void SpawnDefaultActorAttached();
+#endif
+
+	bool GetClonerInitialized() const
+	{
+		return bClonerInitialized;
+	}
+
+protected:
+	//~ Begin UObject
+	virtual void PostInitProperties() override;
+	virtual void PostLoad() override;
+#if WITH_EDITOR
+	virtual void PostEditChangeProperty(FPropertyChangedEvent& InPropertyChangedEvent) override;
+#endif
+	//~ End UObject
+
+	//~ Begin AActor
+	virtual void PostActorCreated() override;
+	/** Allow ticking in editor viewport too */
+	virtual bool ShouldTickIfViewportsOnly() const override
+	{
+		return true;
+	}
+	/** Tick needed to keep the attachment tree updated */
+	virtual void Tick(float DeltaSeconds) override;
+#if WITH_EDITOR
+	virtual FString GetDefaultActorLabel() const override;
+#endif
+	//~ End AActor
+
+	UCEClonerComponent* GetClonerComponent() const
+	{
+		return ClonerComponent;
+	}
+
+	UNiagaraDataInterfaceCurve* GetLifetimeScaleCurveDI() const
+	{
+		return LifetimeScaleCurveDIWeak.Get();
+	}
+
+	/** Will force a system update to refresh user parameters */
+	void RequestClonerUpdate(bool bInImmediate = false);
+
+	void UpdateLayoutOptions();
+	void UpdateClonerEffectors();
+
+	/** Used by effector actors to apply transformations to this cloner instances */
+	const FCEClonerEffectorDataInterfaces* GetEffectorDataInterfaces() const;
+
+	/** Register a new effector and return the index of this effector to match niagara data interface */
+	int32 RegisterEffector(ACEEffectorActor* InEffector);
+
+	/** Unregister the effector and reorganize the niagara data interface to match with the new array indexes */
+	bool UnregisterEffector(ACEEffectorActor* InEffector);
+
+	/** Checks if an effector is registered with this cloner */
+	bool IsEffectorRegistered(const ACEEffectorActor* InEffector) const;
+
+	/** Gets the index of an effector in this cloner */
+	int32 GetEffectorIndex(ACEEffectorActor* InEffector) const;
+
+	/** For each valid effector that has this cloner registered */
+	void ForEachEffector(TFunctionRef<bool(ACEEffectorActor*, int32)> InFunction);
+
+	void OnClonerTransformed(USceneComponent*, EUpdateTransformFlags, ETeleportType);
+	void OnClonerMeshUpdated();
+	void OnClonerSystemChanged();
+
+	void OnEnabledChanged();
+	void OnMeshRenderModeChanged();
+	void OnMeshRendererOptionsChanged();
+	void OnDefaultMeshesChanged();
+	void OnSeedChanged();
+	void OnProgressChanged();
+	void OnDeltaStepChanged();
+	void OnRangeOptionsChanged();
+	void OnLayoutNameChanged();
+	void OnSpawnOptionsChanged();
+	void OnLifetimeScaleCurveChanged();
+	void OnLifetimeOptionsChanged();
+
+	/** Update sprite visibility of this cloner */
+	void OnVisualizerSpriteVisibleChanged();
+
+	bool IsClonerValid() const;
+
+#if WITH_EDITOR
+	void OnReduceMotionGhostingChanged();
+	void OnCVarChanged();
+#endif
+
+	/** Is this cloner enabled/disabled */
+	UPROPERTY(EditAnywhere, BlueprintReadWrite, Setter="SetEnabled", Getter="GetEnabled", Category="Cloner")
+	bool bEnabled = true;
+
+	/** Interval to update the attachment tree and update the cloner meshes */
+	UPROPERTY(EditAnywhere, BlueprintReadWrite, Setter="SetTreeUpdateInterval", Getter="GetTreeUpdateInterval", Category="Cloner", AdvancedDisplay, meta=(ClampMin="0.0"))
+	float TreeUpdateInterval = 0.2f;
+
+	/** Cloner instance seed for random deterministic patterns */
+	UPROPERTY(EditAnywhere, BlueprintReadWrite, Setter="SetSeed", Getter="GetSeed", Category="Cloner")
+	int32 Seed = 0;
+
+	/** Indicates how we select the mesh to render on each clones */
+	UPROPERTY(EditAnywhere, BlueprintReadWrite, Setter="SetMeshRenderMode", Getter="GetMeshRenderMode", Category="Renderer")
+	ECEClonerMeshRenderMode MeshRenderMode = ECEClonerMeshRenderMode::Iterate;
+
+	/** Mode to indicate how clones facing is determined */
+	UPROPERTY(EditAnywhere, BlueprintReadWrite, Setter="SetMeshFacingMode", Getter="GetMeshFacingMode", Category="Renderer")
+	ENiagaraMeshFacingMode MeshFacingMode = ENiagaraMeshFacingMode::Default;
+
+	/** Whether clones cast shadow, disabling will result in better performance */
+	UPROPERTY(EditAnywhere, BlueprintReadWrite, Setter="SetMeshCastShadows", Getter="GetMeshCastShadows", Category="Renderer")
+	bool bMeshCastShadows = true;
+
+	/** When nothing is attached to the cloner, these meshes are used as default */
+	UPROPERTY(EditAnywhere, BlueprintReadWrite, Setter="SetDefaultMeshes", Getter="GetDefaultMeshes", BlueprintSetter="BP_SetDefaultMeshes", BlueprintGetter="BP_GetDefaultMeshes", Category="Renderer")
+	TArray<TObjectPtr<UStaticMesh>> DefaultMeshes;
+
+	/** How many times do we spawn clones */
+	UPROPERTY(EditAnywhere, BlueprintReadWrite, Setter="SetSpawnLoopMode", Getter="GetSpawnLoopMode", DisplayName="LoopMode", Category="Spawn")
+	ECEClonerSpawnLoopMode SpawnLoopMode = ECEClonerSpawnLoopMode::Once;
+
+	/** Amount of spawn iterations for clones */
+	UPROPERTY(EditAnywhere, BlueprintReadWrite, Setter="SetSpawnLoopIterations", Getter="GetSpawnLoopIterations", DisplayName="LoopIterations", Category="Spawn", meta=(ClampMin="1", EditCondition="SpawnLoopMode == ECEClonerSpawnLoopMode::Multiple", EditConditionHides))
+	int32 SpawnLoopIterations = 1;
+
+	/** Interval/Duration of spawn for clones */
+	UPROPERTY(EditAnywhere, BlueprintReadWrite, Setter="SetSpawnLoopInterval", Getter="GetSpawnLoopInterval", DisplayName="LoopInterval", Category="Spawn", meta=(ClampMin="0"))
+	float SpawnLoopInterval = 1.f;
+
+	/** How does spawn occurs */
+	UPROPERTY(EditAnywhere, BlueprintReadWrite, Setter="SetSpawnBehaviorMode", Getter="GetSpawnBehaviorMode", DisplayName="BehaviorMode", Category="Spawn")
+	ECEClonerSpawnBehaviorMode SpawnBehaviorMode = ECEClonerSpawnBehaviorMode::Instant;
+
+	/** How many clones to spawn each seconds */
+	UPROPERTY(EditAnywhere, BlueprintReadWrite, Setter="SetSpawnRate", Getter="GetSpawnRate", DisplayName="Rate", Category="Spawn", meta=(ClampMin="0", EditCondition="SpawnBehaviorMode == ECEClonerSpawnBehaviorMode::Rate", EditConditionHides))
+	float SpawnRate = 1.f;
+
+	/** Do we destroy the clones after a specific duration */
+	UPROPERTY(EditAnywhere, BlueprintReadWrite, Setter="SetLifetimeEnabled", Getter="GetLifetimeEnabled", DisplayName="Enabled", Category="Lifetime")
+	bool bLifetimeEnabled = false;
+
+	/** Minimum lifetime for a clone */
+	UPROPERTY(EditAnywhere, BlueprintReadWrite, Setter="SetLifetimeMin", Getter="GetLifetimeMin", DisplayName="Min", Category="Lifetime", meta=(ClampMin="0", EditCondition="bLifetimeEnabled", EditConditionHides))
+	float LifetimeMin = 0.25f;
+
+	/** Maximum lifetime for a clone */
+	UPROPERTY(EditAnywhere, BlueprintReadWrite, Setter="SetLifetimeMax", Getter="GetLifetimeMax", DisplayName="Max", Category="Lifetime", meta=(ClampMin="0", EditCondition="bLifetimeEnabled", EditConditionHides))
+	float LifetimeMax = 1.5f;
+
+	/** Enable scale by lifetime */
+	UPROPERTY(EditAnywhere, BlueprintReadWrite, Setter="SetLifetimeScaleEnabled", Getter="GetLifetimeScaleEnabled", DisplayName="ScaleEnabled", Category="Lifetime", meta=(EditCondition="bLifetimeEnabled", EditConditionHides))
+	bool bLifetimeScaleEnabled = false;
+
+	/** Used to expose the scale curve editor in details panel */
+	UPROPERTY(Transient, DuplicateTransient, TextExportTransient)
+	TWeakObjectPtr<UNiagaraDataInterfaceCurve> LifetimeScaleCurveDIWeak;
+
+	UPROPERTY(Setter="SetLifetimeScaleCurve", Getter="GetLifetimeScaleCurve")
+	FRichCurve LifetimeScaleCurve;
+
+	/** Name of the layout to use */
+	UPROPERTY(EditAnywhere, BlueprintReadWrite, Setter="SetLayoutName", Getter="GetLayoutName", Category="Layout", meta=(GetOptions="GetClonerLayoutNames"))
+	FName LayoutName = NAME_None;
+
+	/** Active layout */
+	UPROPERTY(VisibleInstanceOnly, BlueprintReadOnly, Instanced, Transient, DuplicateTransient, Category="Layout", meta=(DisplayAfter="LayoutName"))
+	TObjectPtr<UCEClonerLayoutBase> ActiveLayout;
+
+	/** Invert progress behaviour */
+	UPROPERTY(EditAnywhere, BlueprintReadWrite, Setter="SetInvertProgress", Getter="GetInvertProgress", Category="Progress")
+	bool bInvertProgress = false;
+
+	/** Changes visibility of instances based on the total count, 1.f = 100% = all instances visible */
+	UPROPERTY(EditAnywhere, BlueprintReadWrite, Setter="SetProgress", Getter="GetProgress", Category="Progress", meta=(ClampMin="0", ClampMax="1"))
+	float Progress = 1.f;
+
+	/** Enable steps to add delta variation on each clone instance */
+	UPROPERTY(EditAnywhere, BlueprintReadWrite, Setter="SetDeltaStepEnabled", Getter="GetDeltaStepEnabled", DisplayName="Enabled", Category="Step")
+	bool bDeltaStepEnabled = false;
+
+	/** Amount of rotation difference between one step and the next one */
+	UPROPERTY(EditAnywhere, BlueprintReadWrite, Setter="SetDeltaStepRotation", Getter="GetDeltaStepRotation", DisplayName="Rotation", Category="Step", meta=(EditCondition="bDeltaStepEnabled", EditConditionHides))
+	FRotator DeltaStepRotation = FRotator(0.f);
+
+	/** Amount of scale difference between one step and the next one */
+	UPROPERTY(EditAnywhere, BlueprintReadWrite, Setter="SetDeltaStepScale", Getter="GetDeltaStepScale", DisplayName="Scale", Category="Step", meta=(AllowPreserveRatio, Delta="0.0001", EditCondition="bDeltaStepEnabled", EditConditionHides))
+	FVector DeltaStepScale = FVector(0.f);
+
+	/** Use random range transforms for each clones */
+	UPROPERTY(EditAnywhere, BlueprintReadWrite, Setter="SetRangeEnabled", Getter="GetRangeEnabled", DisplayName="Enabled", Category="Range")
+	bool bRangeEnabled = false;
+
+	UPROPERTY(EditAnywhere, BlueprintReadWrite, Setter="SetRangeOffsetMin", Getter="GetRangeOffsetMin", DisplayName="OffsetMin", Category="Range", meta=(EditCondition="bRangeEnabled", EditConditionHides))
+	FVector RangeOffsetMin = FVector::ZeroVector;
+
+	UPROPERTY(EditAnywhere, BlueprintReadWrite, Setter="SetRangeOffsetMax", Getter="GetRangeOffsetMax", DisplayName="OffsetMax", Category="Range", meta=(EditCondition="bRangeEnabled", EditConditionHides))
+	FVector RangeOffsetMax = FVector::ZeroVector;
+
+	UPROPERTY(EditAnywhere, BlueprintReadWrite, Setter="SetRangeRotationMin", Getter="GetRangeRotationMin", DisplayName="RotationMin", Category="Range", meta=(EditCondition="bRangeEnabled", EditConditionHides))
+	FRotator RangeRotationMin = FRotator::ZeroRotator;
+
+	UPROPERTY(EditAnywhere, BlueprintReadWrite, Setter="SetRangeRotationMax", Getter="GetRangeRotationMax", DisplayName="RotationMax", Category="Range", meta=(EditCondition="bRangeEnabled", EditConditionHides))
+	FRotator RangeRotationMax = FRotator::ZeroRotator;
+
+	UPROPERTY(EditAnywhere, BlueprintReadWrite, Setter="SetRangeScaleUniform", Getter="GetRangeScaleUniform", DisplayName="ScaleUniformEnabled", Category="Range", meta=(EditCondition="bRangeEnabled", EditConditionHides))
+	bool bRangeScaleUniform = false;
+
+	UPROPERTY(EditAnywhere, BlueprintReadWrite, Setter="SetRangeScaleMin", Getter="GetRangeScaleMin", DisplayName="ScaleMin", Category="Range", meta=(AllowPreserveRatio, Delta="0.0001", ClampMin="0", EditCondition="bRangeEnabled && !bRangeScaleUniform", EditConditionHides))
+	FVector RangeScaleMin = FVector::OneVector;
+
+	UPROPERTY(EditAnywhere, BlueprintReadWrite, Setter="SetRangeScaleMax", Getter="GetRangeScaleMax", DisplayName="ScaleMax", Category="Range", meta=(AllowPreserveRatio, Delta="0.0001", ClampMin="0", EditCondition="bRangeEnabled && !bRangeScaleUniform", EditConditionHides))
+	FVector RangeScaleMax = FVector::OneVector;
+
+	UPROPERTY(EditAnywhere, BlueprintReadWrite, Setter="SetRangeScaleUniformMin", Getter="GetRangeScaleUniformMin", DisplayName="ScaleMin", Category="Range", meta=(Delta="0.0001", ClampMin="0", EditCondition="bRangeEnabled && bRangeScaleUniform", EditConditionHides))
+	float RangeScaleUniformMin = 1.f;
+
+	UPROPERTY(EditAnywhere, BlueprintReadWrite, Setter="SetRangeScaleUniformMax", Getter="GetRangeScaleUniformMax", DisplayName="ScaleMax", Category="Range", meta=(Delta="0.0001", ClampMin="0", EditCondition="bRangeEnabled && bRangeScaleUniform", EditConditionHides))
+	float RangeScaleUniformMax = 1.f;
+
+#if WITH_EDITORONLY_DATA
+	/** Toggle the sprite to visualize and click on this cloner */
+	UPROPERTY(EditAnywhere, BlueprintReadWrite, AdvancedDisplay, Category="Cloner", meta=(AllowPrivateAccess = "true"))
+	bool bVisualizerSpriteVisible = true;
+
+	/** Reduces the r.TSR.ShadingRejection.Flickering.Period from 3 (default) to 1 if enabled to avoid ghosting artifacts when moving */
+	UPROPERTY(EditAnywhere, BlueprintReadWrite, AdvancedDisplay, Category="Cloner", meta=(AllowPrivateAccess = "true"))
+	bool bReduceMotionGhosting = false;
+#endif
+
+private:
+#if WITH_EDITOR
+	/** Hide outline selection when this cloner only is selected in viewport */
+	void OnEditorSelectionChanged(UObject* InSelection);
+
+	TOptional<bool> UseSelectionOutline;
+#endif
+
+	/** Gets all layout names available */
+	UFUNCTION()
+	TArray<FString> GetClonerLayoutNames() const;
+
+	template<
+		typename InLayoutClass
+		UE_REQUIRES(TIsDerivedFrom<InLayoutClass, UCEClonerLayoutBase>::Value)>
+	InLayoutClass* FindOrAddLayout()
+	{
+		return Cast<InLayoutClass>(FindOrAddLayout(InLayoutClass::StaticClass()));
+	}
+
+	/** Find or add a layout by its class */
+	UCEClonerLayoutBase* FindOrAddLayout(TSubclassOf<UCEClonerLayoutBase> InClass);
+
+	/** Find or add a layout by its name */
+	UCEClonerLayoutBase* FindOrAddLayout(FName InLayoutName);
+
+	/** Initiate and perform operation */
+	void InitializeCloner();
+
+	/** Effectors that applies to this cloner, index match niagara data interface arrays */
+	UPROPERTY(Transient)
+	TArray<TWeakObjectPtr<ACEEffectorActor>> Effectors;
+
+	/** Cloner niagara component */
+	UPROPERTY()
+	TObjectPtr<UCEClonerComponent> ClonerComponent;
+
+	/** Previously used layout instances cached */
+	UPROPERTY()
+	TMap<FName, TObjectPtr<UCEClonerLayoutBase>> LayoutInstances;
+
+	/**
+	* Below properties are deprecated and no longer in use,
+	* they will be migrated to new layout system on load,
+	* And they will be all removed in future version
+	*/
+	UPROPERTY(meta=(DeprecatedProperty, DeprecationMessage="Use LayoutName instead"))
+	ECEClonerLayout Layout_DEPRECATED = ECEClonerLayout::Grid;
+
+	UPROPERTY(meta=(DeprecatedProperty, DeprecationMessage="Use UCEClonerGridLayout instead"))
+	FCEClonerGridLayoutOptions GridOptions_DEPRECATED;
+
+	UPROPERTY(meta=(DeprecatedProperty, DeprecationMessage="Use UCEClonerLineLayout instead"))
+	FCEClonerLineLayoutOptions LineOptions_DEPRECATED;
+
+	UPROPERTY(meta=(DeprecatedProperty, DeprecationMessage="Use UCEClonerCircleLayout instead"))
+	FCEClonerCircleLayoutOptions CircleOptions_DEPRECATED;
+
+	UPROPERTY(meta=(DeprecatedProperty, DeprecationMessage="Use UCEClonerCylinderLayout instead"))
+	FCEClonerCylinderLayoutOptions CylinderOptions_DEPRECATED;
+
+	UPROPERTY(meta=(DeprecatedProperty, DeprecationMessage="Use UCEClonerSphereLayout instead"))
+	FCEClonerSphereLayoutOptions SphereOptions_DEPRECATED;
+
+	UPROPERTY(meta=(DeprecatedProperty, DeprecationMessage="Use UCEClonerHoneycombLayout instead"))
+	FCEClonerHoneycombLayoutOptions HoneycombOptions_DEPRECATED;
+
+	UPROPERTY(meta=(DeprecatedProperty, DeprecationMessage="Use UCEClonerMeshLayout instead"))
+	FCEClonerSampleMeshOptions SampleMeshOptions_DEPRECATED;
+
+	UPROPERTY(meta=(DeprecatedProperty, DeprecationMessage="Use UCEClonerSplineLayout instead"))
+	FCEClonerSampleSplineOptions SampleSplineOptions_DEPRECATED;
+
+	/**
+	 * Used to confirm migration of deprecated properties above to new layout system
+	 * Newly created cloner will not migrate anything and skip this task
+	 */
+	UPROPERTY()
+	bool bDeprecatedPropertiesMigrated = false;
+
+	float TreeUpdateDeltaTime = 0.f;
+
+	bool bNeedsRefresh = false;
+
+	bool bClonerInitialized = false;
+
+#if WITH_EDITOR
+	/** Used for PECP */
+	static const TCEPropertyChangeDispatcher<ACEClonerActor> PropertyChangeDispatcher;
+#endif
+};
diff --git a/Engine/Plugins/Experimental/ClonerEffector/Source/ClonerEffector/Public/Cloner/CEClonerComponent.h b/Engine/Plugins/Experimental/ClonerEffector/Source/ClonerEffector/Public/Cloner/CEClonerComponent.h
new file mode 100644
--- /dev/null
+++ b/Engine/Plugins/Experimental/ClonerEffector/Source/ClonerEffector/Public/Cloner/CEClonerComponent.h
@@ -0,0 +1,124 @@
+// Copyright Epic Games, Inc. All Rights Reserved.
+
+#pragma once
+
+#include "CEClonerEffectorShared.h"
+#include "Layouts/CEClonerLayoutBase.h"
+#include "NiagaraComponent.h"
+#include "CEClonerComponent.generated.h"
+
+class UCEClonerLayoutBase;
+class UMaterialInterface;
+
+UCLASS(MinimalAPI, Within=CEClonerActor, HideCategories=(Niagara))
+class UCEClonerComponent : public UNiagaraComponent
+{
+	GENERATED_BODY()
+
+public:
+	/** Called when meshes have been updated */
+	DECLARE_MULTICAST_DELEGATE(FOnClonerMeshUpdated)
+	static FOnClonerMeshUpdated OnClonerMeshUpdated;
+
+	/** Called when cloner layout system is loaded */
+	DECLARE_MULTICAST_DELEGATE_OneParam(FOnClonerSystemLoaded, UCEClonerLayoutBase* /** InLayout */)
+	static FOnClonerSystemLoaded OnClonerSystemLoaded;
+
+	UCEClonerComponent();
+
+	/**
+	 * Triggers an update of the attachment tree to detect updated items
+	 * If reset is true, clears the attachment tree and rebuilds it otherwise diff update
+	 */
+	void UpdateClonerAttachmentTree(bool bInReset = false);
+
+	/** Called to trigger an update of cloner rendering state if tree */
+	void UpdateClonerRenderState();
+
+	/** Sets the layout to use for this cloner simulation */
+	bool SetClonerActiveLayout(UCEClonerLayoutBase* InLayout);
+
+	UCEClonerLayoutBase* GetClonerActiveLayout() const
+	{
+		return ActiveLayout;
+	}
+
+	/** Forces a refresh of the active system parameters in niagara store */
+	void RefreshUserParameters() const;
+
+protected:
+	//~ Begin UObject
+	virtual void PostLoad() override;
+	//~ End UObject
+
+	void OnRenderStateDirty(UActorComponent& InComponent);
+
+	void UpdateAttachmentTree();
+	void UpdateActorAttachment(AActor* InActor, AActor* InParent);
+
+	void BindActorDelegates(AActor* InActor);
+	void UnbindActorDelegates(AActor* InActor) const;
+
+	void OnTransformUpdated(USceneComponent* InUpdatedComponent, EUpdateTransformFlags InUpdateTransformFlags, ETeleportType InTeleport);
+
+	UFUNCTION()
+	void OnActorDestroyed(AActor* InDestroyedActor);
+
+#if WITH_EDITOR
+	void OnActorPropertyChanged(UObject* InObject, FPropertyChangedEvent& InPropertyChangedEvent);
+#endif
+
+	/** Called when a material has changed in any cloned actors */
+	void OnMaterialChanged(UObject* InObject);
+
+	/** Called when a mesh has been updated and cloner needs to update instances */
+	void OnMeshChanged(UStaticMeshComponent*, AActor* InActor);
+
+	/** Get cloner direct children in the correct order */
+	void GetOrderedRootActors(TArray<AActor*>& OutActors) const;
+
+	/** Get the root cloner actor (direct child) for a specific actor */
+	AActor* GetRootActor(AActor* InActor) const;
+
+	/** Resets the root baked static mesh to be regenerated later */
+	void InvalidateBakedStaticMesh(AActor* InActor);
+
+	/** Runs async update to rebuild dirty meshes */
+	void UpdateDirtyMeshesAsync();
+	void OnDirtyMeshesUpdated();
+
+	/** Merges all primitive components from an actor to dynamic mesh, does not recurse */
+	void UpdateActorBakedDynamicMesh(AActor* InActor);
+
+	/** Merges all baked dynamic meshes from children and self into one static mesh to use as niagara mesh input */
+	void UpdateRootActorBakedStaticMesh(AActor* InRootActor);
+
+	/** Gets all attachment items based on an actor, will recurse */
+	void GetActorAttachmentItems(AActor* InActor, TArray<FCEClonerAttachmentItem*>& OutAttachmentItems);
+
+	/** Update niagara asset static meshes */
+	void UpdateClonerMeshes();
+
+	/** Checks that all root static meshes are valid */
+	bool IsAllMergedMeshesValid() const;
+
+	/** Only materials transient or part of the content folder can be dirtied, engine or plugins cannot */
+	bool IsMaterialDirtyable(const UMaterialInterface* InMaterial) const;
+
+	/** Attachment tree view */
+	UPROPERTY(Transient, NonTransactional)
+	FCEClonerAttachmentTree ClonerTree;
+
+	/** Current active layout */
+	UPROPERTY(Transient, NonTransactional)
+	TObjectPtr<UCEClonerLayoutBase> ActiveLayout;
+
+	/** Attachment items that are dirty and need an update */
+	TSet<FCEClonerAttachmentItem*> DirtyItemAttachments;
+
+	/** Asset meshes needs update */
+	bool bClonerMeshesDirty = true;
+
+	/** State of the baked dynamic and static mesh creation */
+	bool bClonerMeshesUpdating = false;
+};
\ No newline at end of file
diff --git a/Engine/Plugins/Experimental/ClonerEffector/Source/ClonerEffector/Public/Cloner/Layouts/CEClonerGridLayout.h b/Engine/Plugins/Experimental/ClonerEffector/Source/ClonerEffector/Public/Cloner/Layouts/CEClonerGridLayout.h
new file mode 100644
--- /dev/null
+++ b/Engine/Plugins/Experimental/ClonerEffector/Source/ClonerEffector/Public/Cloner/Layouts/CEClonerGridLayout.h
@@ -0,0 +1,173 @@
+// Copyright Epic Games, Inc. All Rights Reserved.
+
+#pragma once
+
+#include "CEClonerEffectorShared.h"
+#include "CEClonerLayoutBase.h"
+#include "CEPropertyChangeDispatcher.h"
+#include "CEClonerGridLayout.generated.h"
+
+UCLASS(MinimalAPI, BlueprintType)
+class UCEClonerGridLayout : public UCEClonerLayoutBase
+{
+	GENERATED_BODY()
+
+	friend class FAvaClonerActorVisualizer;
+
+public:
+	UCEClonerGridLayout()
+		: UCEClonerLayoutBase(
+			TEXT("Grid")
+			, TEXT("/Script/Niagara.NiagaraSystem'/ClonerEffector/Systems/NS_ClonerGrid.NS_ClonerGrid'")
+		)
+	{}
+
+	UFUNCTION(BlueprintCallable, Category="Cloner|Layout|Grid")
+	CLONEREFFECTOR_API void SetCountX(int32 InCountX);
+
+	UFUNCTION(BlueprintPure, Category="Cloner|Layout|Grid")
+	int32 GetCountX() const
+	{
+		return CountX;
+	}
+
+	UFUNCTION(BlueprintCallable, Category="Cloner|Layout|Grid")
+	CLONEREFFECTOR_API void SetCountY(int32 InCountY);
+
+	UFUNCTION(BlueprintPure, Category="Cloner|Layout|Grid")
+	int32 GetCountY() const
+	{
+		return CountY;
+	}
+
+	UFUNCTION(BlueprintCallable, Category="Cloner|Layout|Grid")
+	CLONEREFFECTOR_API void SetCountZ(int32 InCountZ);
+
+	UFUNCTION(BlueprintPure, Category="Cloner|Layout|Grid")
+	int32 GetCountZ() const
+	{
+		return CountZ;
+	}
+
+	UFUNCTION(BlueprintCallable, Category="Cloner|Layout|Grid")
+	CLONEREFFECTOR_API void SetSpacingX(float InSpacingX);
+
+	UFUNCTION(BlueprintPure, Category="Cloner|Layout|Grid")
+	float GetSpacingX() const
+	{
+		return SpacingX;
+	}
+
+	UFUNCTION(BlueprintCallable, Category="Cloner|Layout|Grid")
+	CLONEREFFECTOR_API void SetSpacingY(float InSpacingY);
+
+	UFUNCTION(BlueprintPure, Category="Cloner|Layout|Grid")
+	float GetSpacingY() const
+	{
+		return SpacingY;
+	}
+
+	UFUNCTION(BlueprintCallable, Category="Cloner|Layout|Grid")
+	CLONEREFFECTOR_API void SetSpacingZ(float InSpacingZ);
+
+	UFUNCTION(BlueprintPure, Category="Cloner|Layout|Grid")
+	float GetSpacingZ() const
+	{
+		return SpacingZ;
+	}
+
+	UFUNCTION(BlueprintCallable, Category="Cloner|Layout|Grid")
+	CLONEREFFECTOR_API void SetConstraint(ECEClonerGridConstraint InConstraint);
+
+	UFUNCTION(BlueprintPure, Category="Cloner|Layout|Grid")
+	ECEClonerGridConstraint GetConstraint() const
+	{
+		return Constraint;
+	}
+
+	UFUNCTION(BlueprintCallable, Category="Cloner|Layout|Grid")
+	CLONEREFFECTOR_API void SetInvertConstraint(bool bInInvertConstraint);
+
+	UFUNCTION(BlueprintPure, Category="Cloner|Layout|Grid")
+	bool GetInvertConstraint() const
+	{
+		return bInvertConstraint;
+	}
+
+	UFUNCTION(BlueprintCallable, Category="Cloner|Layout|Grid")
+	CLONEREFFECTOR_API void SetSphereConstraint(const FCEClonerGridConstraintSphere& InConstraint);
+
+	UFUNCTION(BlueprintPure, Category="Cloner|Layout|Grid")
+	const FCEClonerGridConstraintSphere& GetSphereConstraint() const
+	{
+		return SphereConstraint;
+	}
+
+	UFUNCTION(BlueprintCallable, Category="Cloner|Layout|Grid")
+	CLONEREFFECTOR_API void SetCylinderConstraint(const FCEClonerGridConstraintCylinder& InConstraint);
+
+	UFUNCTION(BlueprintPure, Category="Cloner|Layout|Grid")
+	const FCEClonerGridConstraintCylinder& GetCylinderConstraint() const
+	{
+		return CylinderConstraint;
+	}
+
+	UFUNCTION(BlueprintCallable, Category="Cloner|Layout|Grid")
+	CLONEREFFECTOR_API void SetTextureConstraint(const FCEClonerGridConstraintTexture& InConstraint);
+
+	UFUNCTION(BlueprintPure, Category="Cloner|Layout|Grid")
+	const FCEClonerGridConstraintTexture& GetTextureConstraint() const
+	{
+		return TextureConstraint;
+	}
+
+protected:
+	//~ Begin UObject
+#if WITH_EDITOR
+	virtual void PostEditChangeProperty(FPropertyChangedEvent& InPropertyChangedEvent) override;
+#endif
+	//~ End UObject
+
+	//~ Begin UCEClonerLayoutBase
+	virtual void OnLayoutParametersChanged(UCEClonerComponent* InComponent) override;
+	//~ End UCEClonerLayoutBase
+
+	UPROPERTY(EditInstanceOnly, BlueprintReadWrite, Setter="SetCountX", Getter="GetCountX", Category="Layout")
+	int32 CountX = 3;
+
+	UPROPERTY(EditInstanceOnly, BlueprintReadWrite, Setter="SetCountY", Getter="GetCountY", Category="Layout")
+	int32 CountY = 3;
+
+	UPROPERTY(EditInstanceOnly, BlueprintReadWrite, Setter="SetCountZ", Getter="GetCountZ", Category="Layout")
+	int32 CountZ = 3;
+
+	UPROPERTY(EditInstanceOnly, BlueprintReadWrite, Setter="SetSpacingX", Getter="GetSpacingX", Category="Layout")
+	float SpacingX = 105.f;
+
+	UPROPERTY(EditInstanceOnly, BlueprintReadWrite, Setter="SetSpacingY", Getter="GetSpacingY", Category="Layout")
+	float SpacingY = 105.f;
+
+	UPROPERTY(EditInstanceOnly, BlueprintReadWrite, Setter="SetSpacingZ", Getter="GetSpacingZ", Category="Layout")
+	float SpacingZ = 105.f;
+
+	UPROPERTY(EditInstanceOnly, BlueprintReadWrite, Setter="SetConstraint", Getter="GetConstraint", Category="Layout")
+	ECEClonerGridConstraint Constraint = ECEClonerGridConstraint::None;
+
+	UPROPERTY(EditInstanceOnly, BlueprintReadWrite, Setter="SetInvertConstraint", Getter="GetInvertConstraint", Category="Layout", meta=(EditCondition="Constraint != ECEClonerGridConstraint::None", EditConditionHides))
+	bool bInvertConstraint = false;
+
+	UPROPERTY(EditInstanceOnly, BlueprintReadWrite, Setter="SetSphereConstraint", Getter="GetSphereConstraint", Category="Layout", meta=(EditCondition="Constraint == ECEClonerGridConstraint::Sphere", EditConditionHides))
+	FCEClonerGridConstraintSphere SphereConstraint;
+
+	UPROPERTY(EditInstanceOnly, BlueprintReadWrite, Setter="SetCylinderConstraint", Getter="GetCylinderConstraint", Category="Layout", meta=(EditCondition="Constraint == ECEClonerGridConstraint::Cylinder", EditConditionHides))
+	FCEClonerGridConstraintCylinder CylinderConstraint;
+
+	UPROPERTY(EditInstanceOnly, BlueprintReadWrite, Setter="SetTextureConstraint", Getter="GetTextureConstraint", Category="Layout", meta=(EditCondition="Constraint == ECEClonerGridConstraint::Texture", EditConditionHides))
+	FCEClonerGridConstraintTexture TextureConstraint;
+
+private:
+#if WITH_EDITOR
+	/** Used for PECP */
+	static const TCEPropertyChangeDispatcher<UCEClonerGridLayout> PropertyChangeDispatcher;
+#endif
+};
\ No newline at end of file
diff --git a/Engine/Plugins/Experimental/ClonerEffector/Source/ClonerEffector/Public/Effector/CEEffectorActor.h b/Engine/Plugins/Experimental/ClonerEffector/Source/ClonerEffector/Public/Effector/CEEffectorActor.h
new file mode 100644
--- /dev/null
+++ b/Engine/Plugins/Experimental/ClonerEffector/Source/ClonerEffector/Public/Effector/CEEffectorActor.h
@@ -0,0 +1,651 @@
+// Copyright Epic Games, Inc. All Rights Reserved.
+
+#pragma once
+
+#include "CEClonerEffectorShared.h"
+#include "CEPropertyChangeDispatcher.h"
+#include "GameFramework/Actor.h"
+#include "CEEffectorActor.generated.h"
+
+class UNiagaraDataChannelWriter;
+class ACEClonerActor;
+class USphereComponent;
+class UBoxComponent;
+
+UCLASS(MinimalAPI, BlueprintType, HideCategories=(Rendering,Replication,Collision,HLOD,Physics,Networking,Input,Actor,Cooking,LevelInstance), DisplayName = "Motion Design Effector Actor")
+class ACEEffectorActor : public AActor
+{
+	GENERATED_BODY()
+
+	friend class ACEClonerActor;
+	friend class FCEEditorEffectorDetailCustomization;
+	friend class FAvaEffectorActorVisualizer;
+	friend class UCEEffectorSubsystem;
+
+public:
+	static inline const FString DefaultLabel = TEXT("Effector");
+
+	ACEEffectorActor();
+
+	CLONEREFFECTOR_API void LinkCloner(ACEClonerActor* InCloner);
+	CLONEREFFECTOR_API void UnlinkCloner(ACEClonerActor* InCloner);
+
+	UFUNCTION(BlueprintCallable, Category="Effector")
+	CLONEREFFECTOR_API void SetType(ECEClonerEffectorType InType);
+
+	UFUNCTION(BlueprintPure, Category="Effector")
+	ECEClonerEffectorType GetType() const
+	{
+		return Type;
+	}
+
+	UFUNCTION(BlueprintCallable, Category="Effector")
+	CLONEREFFECTOR_API void SetEasing(ECEClonerEasing InEasing);
+
+	UFUNCTION(BlueprintPure, Category="Effector")
+	ECEClonerEasing GetEasing() const
+	{
+		return Easing;
+	}
+
+	UFUNCTION(BlueprintCallable, Category="Effector")
+	CLONEREFFECTOR_API void SetMagnitude(float InMagnitude);
+
+	UFUNCTION(BlueprintPure, Category="Effector")
+	float GetMagnitude() const
+	{
+		return Magnitude;
+	}
+
+	UFUNCTION(BlueprintCallable, Category="Effector")
+	CLONEREFFECTOR_API void SetOuterRadius(float InRadius);
+
+	UFUNCTION(BlueprintPure, Category="Effector")
+	float GetOuterRadius() const
+	{
+		return OuterRadius;
+	}
+
+	UFUNCTION(BlueprintCallable, Category="Effector")
+	CLONEREFFECTOR_API void SetInnerRadius(float InRadius);
+
+	UFUNCTION(BlueprintPure, Category="Effector")
+	float GetInnerRadius() const
+	{
+		return InnerRadius;
+	}
+
+	UFUNCTION(BlueprintCallable, Category="Effector")
+	CLONEREFFECTOR_API void SetInnerExtent(const FVector& InExtent);
+
+	UFUNCTION(BlueprintPure, Category="Effector")
+	FVector GetInnerExtent() const
+	{
+		return InnerExtent;
+	}
+
+	UFUNCTION(BlueprintCallable, Category="Effector")
+	CLONEREFFECTOR_API void SetOuterExtent(const FVector& InExtent);
+
+	UFUNCTION(BlueprintPure, Category="Effector")
+	FVector GetOuterExtent() const
+	{
+		return OuterExtent;
+	}
+
+	UFUNCTION(BlueprintCallable, Category="Effector")
+	CLONEREFFECTOR_API void SetPlaneSpacing(float InSpacing);
+
+	UFUNCTION(BlueprintPure, Category="Effector")
+	float GetPlaneSpacing() const
+	{
+		return PlaneSpacing;
+	}
+
+	UFUNCTION(BlueprintCallable, Category="Effector")
+	CLONEREFFECTOR_API void SetOffset(const FVector& InOffset);
+
+	UFUNCTION(BlueprintPure, Category="Effector")
+	FVector GetOffset() const
+	{
+		return Offset;
+	}
+
+	UFUNCTION(BlueprintCallable, Category="Effector")
+	CLONEREFFECTOR_API void SetRotation(const FRotator& InRotation);
+
+	UFUNCTION(BlueprintPure, Category="Effector")
+	FRotator GetRotation() const
+	{
+		return Rotation;
+	}
+
+	UFUNCTION(BlueprintCallable, Category="Effector")
+	CLONEREFFECTOR_API void SetScale(const FVector& InScale);
+
+	UFUNCTION(BlueprintPure, Category="Effector")
+	FVector GetScale() const
+	{
+		return Scale;
+	}
+
+	UFUNCTION(BlueprintCallable, Category="Effector")
+	CLONEREFFECTOR_API void SetEnabled(bool bInEnable);
+
+	UFUNCTION(BlueprintPure, Category="Effector")
+	bool GetEnabled() const
+	{
+		return bEnabled;
+	}
+
+	UFUNCTION(BlueprintCallable, Category="Effector")
+	CLONEREFFECTOR_API void SetVisualizerThickness(float InThickness);
+
+	UFUNCTION(BlueprintPure, Category="Effector")
+	float GetVisualizerThickness() const
+	{
+		return VisualizerThickness;
+	}
+
+	UFUNCTION(BlueprintCallable, Category="Effector")
+	CLONEREFFECTOR_API void SetMode(ECEClonerEffectorMode InMode);
+
+	UFUNCTION(BlueprintPure, Category="Effector")
+	ECEClonerEffectorMode GetMode() const
+	{
+		return Mode;
+	}
+
+	UFUNCTION(BlueprintCallable, Category="Effector")
+	CLONEREFFECTOR_API void SetTargetActor(AActor* InTargetActor);
+
+	UFUNCTION(BlueprintPure, Category="Effector")
+	AActor* GetTargetActor() const
+	{
+		return TargetActorWeak.Get();
+	}
+
+	UFUNCTION(BlueprintCallable, Category="Cloner")
+	CLONEREFFECTOR_API void SetOrientationForceEnabled(bool bInForceEnabled);
+
+	UFUNCTION(BlueprintPure, Category="Cloner")
+	bool GetOrientationForceEnabled() const
+	{
+		return bOrientationForceEnabled;
+	}
+
+	UFUNCTION(BlueprintCallable, Category="Cloner")
+	CLONEREFFECTOR_API void SetOrientationForceRate(float InForceOrientationRate);
+
+	UFUNCTION(BlueprintPure, Category="Cloner")
+	float GetOrientationForceRate() const
+	{
+		return OrientationForceRate;
+	}
+
+	UFUNCTION(BlueprintCallable, Category="Cloner")
+	CLONEREFFECTOR_API void SetOrientationForceMin(const FVector& InForceOrientationMin);
+
+	UFUNCTION(BlueprintPure, Category="Cloner")
+	const FVector& GetOrientationForceMin() const
+	{
+		return OrientationForceMin;
+	}
+
+	UFUNCTION(BlueprintCallable, Category="Cloner")
+	CLONEREFFECTOR_API void SetOrientationForceMax(const FVector& InForceOrientationMax);
+
+	UFUNCTION(BlueprintPure, Category="Cloner")
+	const FVector& GetOrientationForceMax() const
+	{
+		return OrientationForceMax;
+	}
+
+	UFUNCTION(BlueprintCallable, Category="Cloner")
+	CLONEREFFECTOR_API void SetVortexForceEnabled(bool bInForceEnabled);
+
+	UFUNCTION(BlueprintPure, Category="Cloner")
+	bool GetVortexForceEnabled() const
+	{
+		return bVortexForceEnabled;
+	}
+
+	UFUNCTION(BlueprintCallable, Category="Cloner")
+	CLONEREFFECTOR_API void SetVortexForceAmount(float InForceVortexAmount);
+
+	UFUNCTION(BlueprintPure, Category="Cloner")
+	float GetVortexForceAmount() const
+	{
+		return VortexForceAmount;
+	}
+
+	UFUNCTION(BlueprintCallable, Category="Cloner")
+	CLONEREFFECTOR_API void SetVortexForceAxis(const FVector& InForceVortexAxis);
+
+	UFUNCTION(BlueprintPure, Category="Cloner")
+	const FVector& GetVortexForceAxis() const
+	{
+		return VortexForceAxis;
+	}
+
+	UFUNCTION(BlueprintCallable, Category="Cloner")
+	CLONEREFFECTOR_API void SetCurlNoiseForceEnabled(bool bInForceEnabled);
+
+	UFUNCTION(BlueprintPure, Category="Cloner")
+	bool GetCurlNoiseForceEnabled() const
+	{
+		return bCurlNoiseForceEnabled;
+	}
+
+	UFUNCTION(BlueprintCallable, Category="Cloner")
+	CLONEREFFECTOR_API void SetCurlNoiseForceStrength(float InForceCurlNoiseStrength);
+
+	UFUNCTION(BlueprintPure, Category="Cloner")
+	float GetCurlNoiseForceStrength() const
+	{
+		return CurlNoiseForceStrength;
+	}
+
+	UFUNCTION(BlueprintCallable, Category="Cloner")
+	CLONEREFFECTOR_API void SetCurlNoiseForceFrequency(float InForceCurlNoiseFrequency);
+
+	UFUNCTION(BlueprintPure, Category="Cloner")
+	float GetCurlNoiseForceFrequency() const
+	{
+		return CurlNoiseForceFrequency;
+	}
+
+	UFUNCTION(BlueprintCallable, Category="Cloner")
+	CLONEREFFECTOR_API void SetAttractionForceEnabled(bool bInForceEnabled);
+
+	UFUNCTION(BlueprintPure, Category="Cloner")
+	bool GetAttractionForceEnabled() const
+	{
+		return bAttractionForceEnabled;
+	}
+
+	UFUNCTION(BlueprintCallable, Category="Cloner")
+	CLONEREFFECTOR_API void SetAttractionForceStrength(float InForceStrength);
+
+	UFUNCTION(BlueprintPure, Category="Cloner")
+	float GetAttractionForceStrength() const
+	{
+		return AttractionForceStrength;
+	}
+
+	UFUNCTION(BlueprintCallable, Category="Cloner")
+	CLONEREFFECTOR_API void SetAttractionForceFalloff(float InForceFalloff);
+
+	UFUNCTION(BlueprintPure, Category="Cloner")
+	float GetAttractionForceFalloff() const
+	{
+		return AttractionForceFalloff;
+	}
+
+	UFUNCTION(BlueprintCallable, Category="Cloner")
+	CLONEREFFECTOR_API void SetGravityForceEnabled(bool bInForceEnabled);
+
+	UFUNCTION(BlueprintPure, Category="Cloner")
+	bool GetGravityForceEnabled() const
+	{
+		return bGravityForceEnabled;
+	}
+
+	UFUNCTION(BlueprintCallable, Category="Cloner")
+	CLONEREFFECTOR_API void SetGravityForceAcceleration(const FVector& InAcceleration);
+
+	UFUNCTION(BlueprintPure, Category="Cloner")
+	FVector GetGravityForceAcceleration() const
+	{
+		return GravityForceAcceleration;
+	}
+
+	UFUNCTION(BlueprintCallable, Category="Cloner")
+	CLONEREFFECTOR_API void SetLocationStrength(const FVector& InStrength);
+
+	UFUNCTION(BlueprintPure, Category="Cloner")
+	FVector GetLocationStrength() const
+	{
+		return LocationStrength;
+	}
+
+	UFUNCTION(BlueprintCallable, Category="Cloner")
+	CLONEREFFECTOR_API void SetRotationStrength(const FRotator& InStrength);
+
+	UFUNCTION(BlueprintPure, Category="Cloner")
+	FRotator GetRotationStrength() const
+	{
+		return RotationStrength;
+	}
+
+	UFUNCTION(BlueprintCallable, Category="Cloner")
+	CLONEREFFECTOR_API void SetScaleStrength(const FVector& InStrength);
+
+	UFUNCTION(BlueprintPure, Category="Cloner")
+	FVector GetScaleStrength() const
+	{
+		return ScaleStrength;
+	}
+
+	UFUNCTION(BlueprintCallable, Category="Cloner")
+	CLONEREFFECTOR_API void SetPan(const FVector& InPan);
+
+	UFUNCTION(BlueprintPure, Category="Cloner")
+	FVector GetPan() const
+	{
+		return Pan;
+	}
+
+	UFUNCTION(BlueprintCallable, Category="Cloner")
+	CLONEREFFECTOR_API void SetFrequency(float InFrequency);
+
+	UFUNCTION(BlueprintPure, Category="Cloner")
+	float GetFrequency() const
+	{
+		return Frequency;
+	}
+
+#if WITH_EDITOR
+	UFUNCTION(BlueprintCallable, Category="Effector")
+	CLONEREFFECTOR_API void SetVisualizerSpriteVisible(bool bInVisible);
+
+	UFUNCTION(BlueprintPure, Category="Effector")
+	bool GetVisualizerSpriteVisible() const
+	{
+		return bVisualizerSpriteVisible;
+	}
+#endif
+
+protected:
+	//~ Begin AActor
+	virtual void Destroyed() override;
+	virtual void PostActorCreated() override;
+#if WITH_EDITOR
+	virtual FString GetDefaultActorLabel() const override;
+#endif
+	//~ End AActor
+
+	//~ Begin UObject
+#if WITH_EDITOR
+	virtual void PostEditChangeProperty(FPropertyChangedEvent& InPropertyChangedEvent) override;
+	virtual void PostEditImport() override;
+	virtual void PostEditUndo() override;
+#endif
+	virtual void PostLoad() override;
+	virtual void PostDuplicate(EDuplicateMode::Type InDuplicateMode) override;
+	//~ End UObject
+
+	/** For each valid cloner that has this effector registered, loops in the internal cloner array */
+	void ForEachCloner(TFunctionRef<bool(ACEClonerActor*, int32)> InFunction, bool bSkipIdxCheck = false);
+
+	void OnEffectorSubsystemInitialized(const UWorld* World);
+
+	/** Registers this effector to data channel */
+	void RegisterToChannel();
+
+	/** Called when this effector identifier changed to update linked cloners */
+	void OnEffectorIdentifierChanged();
+
+	FCEClonerEffectorChannelData& GetEffectorChannelData()
+	{
+		return ChannelData;
+	}
+
+	void OnEffectorTransformed(USceneComponent* InUpdatedComponent, EUpdateTransformFlags InUpdateTransformFlags, ETeleportType InTeleport);
+
+	void OnClonerUpdated(ACEClonerActor* InCloner, int32 InEffectorIdx) {}
+	void OnClonerLinked(ACEClonerActor* InCloner, int32 InEffectorIdx);
+	void OnClonerUnlinked(ACEClonerActor* InCloner, int32 InEffectorIdx);
+
+	/** Will update cloner system only if effector is enabled or force update is true, update will happen next tick */
+	void RefreshClonerParameters(ACEClonerActor* InCloner, bool bInForce = false) const;
+	void RefreshClonersParameters(bool bInForce = false);
+
+	/** Will register this effector with linked clones */
+	void RegisterToCloners();
+
+	/** Will update all values of this effector on all clones */
+	void OnEffectorChanged();
+
+	/** Called when effector is enabled/disabled */
+	void OnEnabledChanged();
+
+	/** Called when we disable this effector */
+	void OnEffectorDisabled();
+
+	/** Called when cloners set is updated */
+	void OnClonersChanged();
+
+	/** Update mode for this effector */
+	void OnModeChanged();
+
+	/** Update shape type for this effector */
+	void OnTypeChanged();
+
+	/** Update values for box type effector */
+	void OnBoxChanged();
+
+	/** Update values for sphere type effector */
+	void OnSphereChanged();
+
+	/** Update values for plane type effector */
+	void OnPlaneChanged();
+
+	/** Called when transform of effector has changed */
+	void OnTransformChanged();
+
+	/** Update values for easing */
+	void OnEasingChanged();
+
+	/** Update values for magnitude */
+	void OnMagnitudeChanged();
+
+	/** Called when default mode options are changed */
+	void OnTransformOptionsChanged();
+
+	/** Update actor targeted by this effector on cloners */
+	void OnTargetActorChanged();
+
+	/** Called when target mode options are changed */
+	void OnTargetOptionsChanged();
+
+	/** Called when noise field mode options are changed */
+	void OnNoiseFieldOptionsChanged();
+
+	/** Called when force options are changed */
+	void OnForceOptionsChanged();
+
+	/** Called when force enabled state changed */
+	void OnForceEnabledChanged();
+
+	/** Update thickness of components of this effector */
+	void OnVisualizerThicknessChanged();
+
+	/** Update sprite visibility of this effector */
+	void OnVisualizerSpriteVisibleChanged();
+
+	/** Is this effector enabled/disabled on linked cloners */
+	UPROPERTY(EditAnywhere, BlueprintReadWrite, Setter="SetEnabled", Getter="GetEnabled", Category="Effector")
+	bool bEnabled = true;
+
+	/** The ratio effect of the effector on clones */
+	UPROPERTY(EditAnywhere, BlueprintReadWrite, Setter="SetMagnitude", Getter="GetMagnitude", Category="Effector", meta=(ClampMin="0", ClampMax="1"))
+	float Magnitude = 1.f;
+
+	/** Cloners affected by this effector */
+	UPROPERTY(EditAnywhere, Category="Effector", meta=(AllowPrivateAccess = "true"))
+	TSet<TWeakObjectPtr<ACEClonerActor>> Cloners;
+
+	/** Type of effector to apply on cloners instances */
+	UPROPERTY(EditAnywhere, BlueprintReadWrite, Setter="SetType", Getter="GetType", Category="Type")
+	ECEClonerEffectorType Type = ECEClonerEffectorType::Sphere;
+
+	/** Weight easing function applied to lerp transforms */
+	UPROPERTY(EditAnywhere, BlueprintReadWrite, Setter="SetEasing", Getter="GetEasing", Category="Type", meta=(EditCondition="Type != ECEClonerEffectorType::Unbound", EditConditionHides))
+	ECEClonerEasing Easing = ECEClonerEasing::Linear;
+
+	/** Inner radius of sphere, all clones inside will be affected with a maximum weight */
+	UPROPERTY(EditAnywhere, BlueprintReadWrite, Setter="SetInnerRadius", Getter="GetInnerRadius", Category="Type", meta=(ClampMin="0", EditCondition="Type == ECEClonerEffectorType::Sphere", EditConditionHides))
+	float InnerRadius = 50.f;
+
+	/** Outer radius of sphere, all clones outside will not be affected */
+	UPROPERTY(EditAnywhere, BlueprintReadWrite, Setter="SetOuterRadius", Getter="GetOuterRadius", Category="Type", meta=(ClampMin="0", EditCondition="Type == ECEClonerEffectorType::Sphere", EditConditionHides))
+	float OuterRadius = 200.f;
+
+	/** Inner extent of box, all clones inside will be affected with a maximum weight */
+	UPROPERTY(EditAnywhere, BlueprintReadWrite, Setter="SetInnerExtent", Getter="GetInnerExtent", Category="Type", meta=(ClampMin="0", EditCondition="Type == ECEClonerEffectorType::Box", EditConditionHides))
+	FVector InnerExtent = FVector(50.f);
+
+	/** Outer extent of box, all clones outside will not be affected */
+	UPROPERTY(EditAnywhere, BlueprintReadWrite, Setter="SetOuterExtent", Getter="GetOuterExtent", Category="Type", meta=(ClampMin="0", EditCondition="Type == ECEClonerEffectorType::Box", EditConditionHides))
+	FVector OuterExtent = FVector(200.f);
+
+	/** Plane spacing, everything inside this zone will be affected */
+	UPROPERTY(EditAnywhere, BlueprintReadWrite, Setter="SetPlaneSpacing", Getter="GetPlaneSpacing", Category="Type", meta=(ClampMin="0", EditCondition="Type == ECEClonerEffectorType::Plane", EditConditionHides))
+	float PlaneSpacing = 200.f;
+
+	/** Mode of effector for each clones instances */
+	UPROPERTY(EditAnywhere, BlueprintReadWrite, Setter="SetMode", Getter="GetMode", Category="Mode")
+	ECEClonerEffectorMode Mode = ECEClonerEffectorMode::Default;
+
+	/** Offset applied on affected clones */
+	UPROPERTY(EditAnywhere, BlueprintReadWrite, Setter="SetOffset", Getter="GetOffset", Category="Mode", meta=(EditCondition="Mode == ECEClonerEffectorMode::Default", EditConditionHides))
+	FVector Offset = FVector::ZeroVector;
+
+	/** Rotation applied on affected clones */
+	UPROPERTY(EditAnywhere, BlueprintReadWrite, Setter="SetRotation", Getter="GetRotation", Category="Mode", meta=(ClampMin="-180", ClampMax="180", UIMin="-180", UIMax="180", EditCondition="Mode == ECEClonerEffectorMode::Default", EditConditionHides))
+	FRotator Rotation = FRotator::ZeroRotator;
+
+	/** Scale applied on affected clones */
+	UPROPERTY(EditAnywhere, BlueprintReadWrite, Setter="SetScale", Getter="GetScale", Category="Mode", meta=(ClampMin="0", AllowPreserveRatio, Delta="0.01", EditCondition="Mode == ECEClonerEffectorMode::Default", EditConditionHides))
+	FVector Scale = FVector::OneVector;
+
+	/** The actor to track when mode is set to target */
+	UPROPERTY(EditAnywhere, BlueprintReadWrite, Category="Mode", meta=(DisplayName="TargetActor", EditCondition="Mode == ECEClonerEffectorMode::Target", EditConditionHides))
+	TWeakObjectPtr<AActor> TargetActorWeak = nullptr;
+
+	UPROPERTY()
+	TWeakObjectPtr<AActor> InternalTargetActorWeak = nullptr;
+
+	/** Amplitude of the noise field for location */
+	UPROPERTY(EditAnywhere, BlueprintReadWrite, Setter="SetLocationStrength", Getter="GetLocationStrength", Category="Mode", meta=(EditCondition="Mode == ECEClonerEffectorMode::NoiseField", EditConditionHides))
+	FVector LocationStrength = FVector::ZeroVector;
+
+	/** Amplitude of the noise field for rotation */
+	UPROPERTY(EditAnywhere, BlueprintReadWrite, Category="Mode", meta=(EditCondition="Mode == ECEClonerEffectorMode::NoiseField", EditConditionHides))
+	FRotator RotationStrength = FRotator::ZeroRotator;
+
+	/** Amplitude of the noise field for scale */
+	UPROPERTY(EditAnywhere, BlueprintReadWrite, Category="Mode", meta=(EditCondition="Mode == ECEClonerEffectorMode::NoiseField", EditConditionHides))
+	FVector ScaleStrength = FVector::OneVector;
+
+	/** Panning to offset the noise field sampling */
+	UPROPERTY(EditAnywhere, BlueprintReadWrite, Setter="SetPan", Getter="GetPan", Category="Mode", meta=(EditCondition="Mode == ECEClonerEffectorMode::NoiseField", EditConditionHides))
+	FVector Pan = FVector::ZeroVector;
+
+	/** Intensity of the noise field */
+	UPROPERTY(EditAnywhere, BlueprintReadWrite, Setter="SetFrequency", Getter="GetFrequency", Category="Mode", meta=(ClampMin="0", EditCondition="Mode == ECEClonerEffectorMode::NoiseField", EditConditionHides))
+	float Frequency = 0.5f;
+
+	/** Enable orientation force to allow each clone instance to rotate around its pivot */
+	UPROPERTY(EditAnywhere, BlueprintReadWrite, Setter="SetOrientationForceEnabled", Getter="GetOrientationForceEnabled", Category="Force")
+	bool bOrientationForceEnabled = false;
+
+	UPROPERTY(EditAnywhere, BlueprintReadWrite, Setter="SetOrientationForceRate", Getter="GetOrientationForceRate", Category="Force", meta=(Delta="0.0001", ClampMin="0", EditCondition="bOrientationForceEnabled", EditConditionHides))
+	float OrientationForceRate = 1.f;
+
+	UPROPERTY(EditAnywhere, BlueprintReadWrite, Setter="SetOrientationForceMin", Getter="GetOrientationForceMin", Category="Force", meta=(EditCondition="bOrientationForceEnabled", EditConditionHides))
+	FVector OrientationForceMin = FVector(-0.1f);
+
+	UPROPERTY(EditAnywhere, BlueprintReadWrite, Setter="SetOrientationForceMax", Getter="GetOrientationForceMax", Category="Force", meta=(EditCondition="bOrientationForceEnabled", EditConditionHides))
+	FVector OrientationForceMax = FVector(0.1f);
+
+	/** Enable vortex force to allow each clone instance to rotate around a specific axis on the cloner pivot */
+	UPROPERTY(EditAnywhere, BlueprintReadWrite, Setter="SetVortexForceEnabled", Getter="GetVortexForceEnabled", Category="Force")
+	bool bVortexForceEnabled = false;
+
+	UPROPERTY(EditAnywhere, BlueprintReadWrite, Setter="SetVortexForceAmount", Getter="GetVortexForceAmount", Category="Force", meta=(EditCondition="bVortexForceEnabled", EditConditionHides))
+	float VortexForceAmount = 10000.f;
+
+	UPROPERTY(EditAnywhere, BlueprintReadWrite, Setter="SetVortexForceAxis", Getter="GetVortexForceAxis", Category="Force", meta=(EditCondition="bVortexForceEnabled", EditConditionHides))
+	FVector VortexForceAxis = FVector::ZAxisVector;
+
+	/** Enable curl noise force to allow each clone instance to add random location variation */
+	UPROPERTY(EditAnywhere, BlueprintReadWrite, Setter="SetCurlNoiseForceEnabled", Getter="GetCurlNoiseForceEnabled", Category="Force")
+	bool bCurlNoiseForceEnabled = false;
+
+	UPROPERTY(EditAnywhere, BlueprintReadWrite, Setter="SetCurlNoiseForceStrength", Getter="GetCurlNoiseForceStrength", Category="Force", meta=(EditCondition="bCurlNoiseForceEnabled", EditConditionHides))
+	float CurlNoiseForceStrength = 1000.f;
+
+	UPROPERTY(EditAnywhere, BlueprintReadWrite, Setter="SetCurlNoiseForceFrequency", Getter="GetCurlNoiseForceFrequency", Category="Force", meta=(EditCondition="bCurlNoiseForceEnabled", EditConditionHides))
+	float CurlNoiseForceFrequency = 10.f;
+
+	/** Enable attraction force to allow each clone instances to gravitate toward a location */
+	UPROPERTY(EditAnywhere, BlueprintReadWrite, Setter="SetAttractionForceEnabled", Getter="GetAttractionForceEnabled", Category="Force")
+	bool bAttractionForceEnabled = false;
+
+	UPROPERTY(EditAnywhere, BlueprintReadWrite, Setter="SetAttractionForceStrength", Getter="GetAttractionForceStrength", Category="Force", meta=(EditCondition="bAttractionForceEnabled", EditConditionHides))
+	float AttractionForceStrength = 1000.f;
+
+	UPROPERTY(EditAnywhere, BlueprintReadWrite, Setter="SetAttractionForceFalloff", Getter="GetAttractionForceFalloff", Category="Force", meta=(EditCondition="bAttractionForceEnabled", EditConditionHides))
+	float AttractionForceFalloff = 0.1f;
+
+	/** Enable gravity force to pull particles based on an acceleration vector */
+	UPROPERTY(EditAnywhere, BlueprintReadWrite, Setter="SetGravityForceEnabled", Getter="GetGravityForceEnabled", Category="Force")
+	bool bGravityForceEnabled = false;
+
+	UPROPERTY(EditAnywhere, BlueprintReadWrite, Setter="SetGravityForceAcceleration", Getter="GetGravityForceAcceleration", Category="Force", meta=(EditCondition="bGravityForceEnabled", EditConditionHides))
+	FVector GravityForceAcceleration = FVector(0, 0, -980.f);
+
+	/** Thickness of components visualizers */
+	UPROPERTY(EditAnywhere, BlueprintReadWrite, AdvancedDisplay, Setter="SetVisualizerThickness", Getter="GetVisualizerThickness", Category="Effector", meta=(ClampMin="0.1", ClampMax="10.0", Delta="0.1"))
+	float VisualizerThickness = 2.f;
+
+#if WITH_EDITORONLY_DATA
+	/** Toggle the sprite to visualize and click on this effector */
+	UPROPERTY(EditAnywhere, BlueprintReadWrite, AdvancedDisplay, Category="Effector")
+	bool bVisualizerSpriteVisible = true;
+#endif
+
+private:
+	void OnTargetActorTransformChanged(USceneComponent* InUpdatedComponent, EUpdateTransformFlags InFlags, ETeleportType InTeleport);
+
+	UFUNCTION()
+	void OnTargetActorDestroyed(AActor* InActor);
+
+	/** Sphere component for sphere maximum weight */
+	UPROPERTY()
+	TObjectPtr<USphereComponent> InnerSphereComponent;
+
+	/** Sphere component for sphere minimum weight */
+	UPROPERTY()
+	TObjectPtr<USphereComponent> OuterSphereComponent;
+
+	/** Box component for box maximum weight */
+	UPROPERTY()
+	TObjectPtr<UBoxComponent> InnerBoxComponent;
+
+	/** Box component for box minimum weight */
+	UPROPERTY()
+	TObjectPtr<UBoxComponent> OuterBoxComponent;
+
+	/** Plane components for plane minimum weight */
+	UPROPERTY()
+	TObjectPtr<UBoxComponent> InnerPlaneComponent;
+
+	/** Plane components for plane maximum weight */
+	UPROPERTY()
+	TObjectPtr<UBoxComponent> OuterPlaneComponent;
+
+	/** Internal cloners array */
+	UPROPERTY(NonTransactional)
+	TSet<TWeakObjectPtr<ACEClonerActor>> InternalCloners;
+
+	/** Transient effector channel data */
+	FCEClonerEffectorChannelData ChannelData;
+
+#if WITH_EDITOR
+	/** Used for PECP */
+	static TCEPropertyChangeDispatcher<ACEEffectorActor> PropertyChangeDispatcher;
+#endif
+};
diff --git a/Engine/Plugins/Experimental/ClonerEffector/Source/ClonerEffector/Public/Effector/CEEffectorComponent.h b/Engine/Plugins/Experimental/ClonerEffector/Source/ClonerEffector/Public/Effector/CEEffectorComponent.h
new file mode 100644
--- /dev/null
+++ b/Engine/Plugins/Experimental/ClonerEffector/Source/ClonerEffector/Public/Effector/CEEffectorComponent.h
@@ -0,0 +1,13 @@
+// Copyright Epic Games, Inc. All Rights Reserved.
+
+#pragma once
+
+#include "Components/SceneComponent.h"
+#include "CEEffectorComponent.generated.h"
+
+/** Class used to define a custom visualizer for this component and the effector actor */
+UCLASS(MinimalAPI, Within=CEEffectorActor)
+class UCEEffectorComponent : public USceneComponent
+{
+	GENERATED_BODY()
+};
\ No newline at end of file
diff --git a/Engine/Plugins/Experimental/ClonerEffector/Source/ClonerEffector/Public/Subsystems/CEClonerSubsystem.h b/Engine/Plugins/Experimental/ClonerEffector/Source/ClonerEffector/Public/Subsystems/CEClonerSubsystem.h
new file mode 100644
--- /dev/null
+++ b/Engine/Plugins/Experimental/ClonerEffector/Source/ClonerEffector/Public/Subsystems/CEClonerSubsystem.h
@@ -0,0 +1,89 @@
+// Copyright Epic Games, Inc. All Rights Reserved.
+
+#pragma once
+
+#include "Cloner/Layouts/CEClonerLayoutBase.h"
+#include "Subsystems/EngineSubsystem.h"
+#include "UObject/WeakObjectPtrTemplates.h"
+#include "CEClonerSubsystem.generated.h"
+
+UCLASS(MinimalAPI)
+class UCEClonerSubsystem : public UEngineSubsystem
+{
+	GENERATED_BODY()
+
+public:
+	static inline constexpr int32 NoFlicker = 1;
+
+#if WITH_EDITOR
+	DECLARE_MULTICAST_DELEGATE(FOnCVarChanged)
+	static FOnCVarChanged& OnCVarChanged()
+	{
+		return OnCVarChangedDelegate;
+	}
+#endif
+
+	DECLARE_MULTICAST_DELEGATE(FOnSubsystemInitialized)
+	static FOnSubsystemInitialized& OnSubsystemInitialized()
+	{
+		 return OnSubsystemInitializedDelegate;
+	}
+
+	/** Get this subsystem instance */
+	CLONEREFFECTOR_API static UCEClonerSubsystem* Get();
+
+	//~ Begin USubsystem
+	virtual void Initialize(FSubsystemCollectionBase& Collection) override;
+	virtual void Deinitialize() override;
+	//~ End USubsystem
+
+	CLONEREFFECTOR_API bool RegisterLayoutClass(const UClass* InClonerLayoutClass);
+	CLONEREFFECTOR_API bool UnregisterLayoutClass(const UClass* InClonerLayoutClass);
+	CLONEREFFECTOR_API bool IsLayoutClassRegistered(const UClass* InClonerLayoutClass);
+
+	DECLARE_DELEGATE_RetVal_OneParam(TArray<AActor*> /** Children */, FOnGetOrderedActors, const AActor* /** InParent */)
+	CLONEREFFECTOR_API void RegisterCustomActorResolver(FOnGetOrderedActors InCustomResolver);
+	CLONEREFFECTOR_API void UnregisterCustomActorResolver();
+	FOnGetOrderedActors& GetCustomActorResolver();
+
+	/** Get available cloner layout names to use in dropdown */
+	TArray<FName> GetLayoutNames() const;
+
+	/** Based on a layout class, find layout name */
+	FName FindLayoutName(TSubclassOf<UCEClonerLayoutBase> InLayoutClass) const;
+
+	/** Creates a new layout instance for a cloner actor */
+	UCEClonerLayoutBase* CreateNewLayout(FName InLayoutName, ACEClonerActor* InClonerActor);
+
+#if WITH_EDITOR
+	void EnableNoFlicker();
+	void DisableNoFlicker();
+	bool IsNoFlickerEnabled() const;
+
+	void OnTSRShadingRejectionFlickeringPeriodChanged(IConsoleVariable* InCVar) const;
+#endif
+
+protected:
+#if WITH_EDITOR
+	static FOnCVarChanged OnCVarChangedDelegate;
+#endif
+
+	CLONEREFFECTOR_API static FOnSubsystemInitialized OnSubsystemInitializedDelegate;
+
+	void ScanForRegistrableClasses();
+
+	/** Linking name to the layout class */
+	UPROPERTY()
+	TMap<FName, TSubclassOf<UCEClonerLayoutBase>> LayoutClasses;
+
+	/** Used to gather ordered actors based on parent */
+	FOnGetOrderedActors ActorResolver;
+
+#if WITH_EDITOR
+	/** Allows to reduce ghosting artifacts when moving clone instances */
+	IConsoleVariable* CVarTSRShadingRejectionFlickeringPeriod = nullptr;
+
+	/** Previous value to restore it when disabled */
+	TOptional<int32> PreviousCVarValue;
+#endif
+};
\ No newline at end of file
diff --git a/Engine/Plugins/Experimental/ClonerEffector/Source/ClonerEffectorEditor/Private/CEEditorModule.cpp b/Engine/Plugins/Experimental/ClonerEffector/Source/ClonerEffectorEditor/Private/CEEditorModule.cpp
new file mode 100644
--- /dev/null
+++ b/Engine/Plugins/Experimental/ClonerEffector/Source/ClonerEffectorEditor/Private/CEEditorModule.cpp
@@ -0,0 +1,37 @@
+﻿// Copyright Epic Games, Inc. All Rights Reserved.
+
+#include "CEEditorModule.h"
+
+#include "Cloner/CEClonerActor.h"
+#include "Cloner/CEEditorClonerDetailCustomization.h"
+#include "Effector/CEEditorEffectorDetailCustomization.h"
+#include "Effector/CEEffectorActor.h"
+#include "Modules/ModuleManager.h"
+#include "PropertyEditorModule.h"
+#include "Styles/CEEditorStyle.h"
+
+void FCEEditorModule::StartupModule()
+{
+	FPropertyEditorModule& PropertyModule = FModuleManager::GetModuleChecked<FPropertyEditorModule>(PropertyEditorName);
+
+	// Load styles
+	FCEEditorStyle::Get();
+
+	// Cloner/effector customization
+	PropertyModule.RegisterCustomClassLayout(ACEEffectorActor::StaticClass()->GetFName(), FOnGetDetailCustomizationInstance::CreateStatic(&FCEEditorEffectorDetailCustomization::MakeInstance));
+	PropertyModule.RegisterCustomClassLayout(ACEClonerActor::StaticClass()->GetFName(), FOnGetDetailCustomizationInstance::CreateStatic(&FCEEditorClonerDetailCustomization::MakeInstance));
+}
+
+void FCEEditorModule::ShutdownModule()
+{
+	if (FModuleManager::Get().IsModuleLoaded(PropertyEditorName))
+	{
+		FPropertyEditorModule& PropertyModule = FModuleManager::GetModuleChecked<FPropertyEditorModule>(PropertyEditorName);
+
+		// Cloner/effector
+		PropertyModule.UnregisterCustomClassLayout(ACEEffectorActor::StaticClass()->GetFName());
+		PropertyModule.UnregisterCustomClassLayout(ACEClonerActor::StaticClass()->GetFName());
+	}
+}
+
+IMPLEMENT_MODULE(FCEEditorModule, ClonerEffectorEditor)
\ No newline at end of file
diff --git a/Engine/Plugins/Experimental/ClonerEffector/Source/ClonerEffectorEditor/Private/CEEditorModule.h b/Engine/Plugins/Experimental/ClonerEffector/Source/ClonerEffectorEditor/Private/CEEditorModule.h
new file mode 100644
--- /dev/null
+++ b/Engine/Plugins/Experimental/ClonerEffector/Source/ClonerEffectorEditor/Private/CEEditorModule.h
@@ -0,0 +1,16 @@
+﻿// Copyright Epic Games, Inc. All Rights Reserved.
+
+#pragma once
+
+#include "Modules/ModuleManager.h"
+
+class FCEEditorModule : public IModuleInterface
+{
+public:
+	static inline constexpr TCHAR PropertyEditorName[] = TEXT("PropertyEditor");
+
+	//~ Begin IModuleInterface
+    virtual void StartupModule() override;
+    virtual void ShutdownModule() override;
+	//~ End IModuleInterface
+};
diff --git a/Engine/Plugins/Experimental/ClonerEffector/Source/ClonerEffectorEditor/Private/Cloner/CEEditorClonerDetailCustomization.cpp b/Engine/Plugins/Experimental/ClonerEffector/Source/ClonerEffectorEditor/Private/Cloner/CEEditorClonerDetailCustomization.cpp
new file mode 100644
--- /dev/null
+++ b/Engine/Plugins/Experimental/ClonerEffector/Source/ClonerEffectorEditor/Private/Cloner/CEEditorClonerDetailCustomization.cpp
@@ -0,0 +1,200 @@
+// Copyright Epic Games, Inc. All Rights Reserved.
+
+#include "Cloner/CEEditorClonerDetailCustomization.h"
+
+#include "Cloner/CEClonerActor.h"
+#include "Cloner/Layouts/CEClonerLayoutBase.h"
+#include "DetailLayoutBuilder.h"
+#include "DetailWidgetRow.h"
+#include "Modules/ModuleManager.h"
+#include "NiagaraDataInterfaceCurve.h"
+#include "PropertyEditorModule.h"
+#include "Widgets/Input/SButton.h"
+
+#define LOCTEXT_NAMESPACE "CEEditorClonerDetailCustomization"
+
+void FCEEditorClonerDetailCustomization::CustomizeDetails(IDetailLayoutBuilder& InDetailBuilder)
+{
+	static const TArray<FName> CategoriesPriority =
+	{
+		TEXT("Cloner"),
+		TEXT("Layout"),
+		TEXT("Renderer"),
+		TEXT("Progress"),
+		TEXT("Step"),
+		TEXT("Range"),
+		TEXT("Spawn"),
+		TEXT("Lifetime"),
+	};
+
+	// Sort categories
+	int32 Order = InDetailBuilder.EditCategory(TEXT("Transform")).GetSortOrder();
+	for (const FName& Category : CategoriesPriority)
+	{
+		IDetailCategoryBuilder& CategoryBuilder = InDetailBuilder.EditCategory(Category);
+		CategoryBuilder.SetSortOrder(++Order);
+	}
+
+	// Gather all cloner actor selected
+	TArray<TWeakObjectPtr<ACEClonerActor>> SelectedClonersWeak = InDetailBuilder.GetSelectedObjectsOfType<ACEClonerActor>();
+
+	LayoutFunctionNames.Empty();
+
+	if (SelectedClonersWeak.IsEmpty())
+	{
+		return;
+	}
+
+	// Look for layout UFunctions and group them by name
+	for (const TWeakObjectPtr<ACEClonerActor>& ClonerWeak : SelectedClonersWeak)
+	{
+		const ACEClonerActor* Cloner = ClonerWeak.Get();
+		if (!Cloner)
+		{
+			continue;
+		}
+
+		// Look for CallInEditor functions in active layout
+		UCEClonerLayoutBase* ActiveLayout = Cloner->GetActiveLayout();
+		if (!ActiveLayout)
+		{
+			continue;
+		}
+
+		// Iterate through all UFunctions in the class and its parent classes.
+		for (UFunction* Function : TFieldRange<UFunction>(ActiveLayout->GetClass(), EFieldIteratorFlags::ExcludeSuper))
+		{
+			// Only CallInEditor function with 0 parameters
+			if (Function && Function->HasMetaData("CallInEditor") && Function->NumParms == 0)
+			{
+				FName FunctionName = Function->GetFName();
+				TMap<TWeakObjectPtr<UObject>, TWeakObjectPtr<UFunction>>& ObjectFunctions = LayoutFunctionNames.FindOrAdd(FunctionName);
+				ObjectFunctions.Add(ActiveLayout, Function);
+			}
+		}
+	}
+
+	if (!LayoutFunctionNames.IsEmpty())
+	{
+		// Add buttons for selected cloners
+		const TSharedPtr<SHorizontalBox> FunctionsWidget = SNew(SHorizontalBox);
+		for (const TPair<FName, TMap<TWeakObjectPtr<UObject>, TWeakObjectPtr<UFunction>>>& LayoutFunctionNamesPair : LayoutFunctionNames)
+		{
+			const FText ButtonLabel = FText::FromString(FName::NameToDisplayString(LayoutFunctionNamesPair.Key.ToString(), false));
+
+			FunctionsWidget->AddSlot()
+			.AutoWidth()
+			[
+				SNew(SButton)
+				.Text(ButtonLabel)
+				.OnClicked(this, &FCEEditorClonerDetailCustomization::OnExecuteFunction, LayoutFunctionNamesPair.Key)
+			];
+		}
+
+		IDetailCategoryBuilder& ClonerCategoryBuilder = InDetailBuilder.EditCategory("Cloner");
+
+		ClonerCategoryBuilder
+			.AddCustomRow(FText::FromString("LayoutFunctions"))
+			.WholeRowContent()
+			[
+				FunctionsWidget.ToSharedRef()
+			];
+	}
+
+	if (SelectedClonersWeak.Num() != 1)
+	{
+		return;
+	}
+
+	const ACEClonerActor* Cloner = SelectedClonersWeak[0].Get();
+
+	if (!Cloner)
+	{
+		return;
+	}
+
+	UNiagaraDataInterfaceCurve* LifetimeScaleCurve = Cloner->GetLifetimeScaleCurveDI();
+
+	if (!LifetimeScaleCurve)
+	{
+		return;
+	}
+
+	const TArray<UObject*> ShowObjects {LifetimeScaleCurve};
+	IDetailCategoryBuilder& LifetimeCategory = InDetailBuilder.EditCategory(TEXT("Lifetime"));
+	LifetimeCategory.SetShowAdvanced(true);
+
+	FAddPropertyParams Params;
+	Params.HideRootObjectNode(true);
+	Params.CreateCategoryNodes(false);
+
+	IDetailPropertyRow* CurveRow = LifetimeCategory.AddExternalObjects(ShowObjects, EPropertyLocation::Advanced, Params);
+	CurveRow->Visibility(MakeAttributeLambda([Cloner]()
+	{
+		return Cloner
+			&& Cloner->GetLifetimeEnabled()
+			&& Cloner->GetLifetimeScaleEnabled()
+			? EVisibility::Visible
+			: EVisibility::Collapsed;
+	}));
+}
+
+FReply FCEEditorClonerDetailCustomization::OnExecuteFunction(FName InFunctionName)
+{
+	if (TMap<TWeakObjectPtr<UObject>, TWeakObjectPtr<UFunction>> const* ObjectFunctions = LayoutFunctionNames.Find(InFunctionName))
+	{
+		for (const TPair<TWeakObjectPtr<UObject>, TWeakObjectPtr<UFunction>>& ObjectFunction : *ObjectFunctions)
+		{
+			UObject* Object = ObjectFunction.Key.Get();
+			UFunction* Function = ObjectFunction.Value.Get();
+
+			if (!Object || !Function)
+			{
+				continue;
+			}
+
+			Object->ProcessEvent(Function, nullptr);
+		}
+	}
+
+	return FReply::Handled();
+}
+
+void FCEEditorClonerDetailCustomization::RegisterCustomSections() const
+{
+	static const FName PropertyEditor("PropertyEditor");
+	FPropertyEditorModule& PropertyModule = FModuleManager::GetModuleChecked<FPropertyEditorModule>(PropertyEditor);
+
+	static const FName ClassName = ACEClonerActor::StaticClass()->GetFName();
+
+	const TSharedRef<FPropertySection> StreamingSection = PropertyModule.FindOrCreateSection(ClassName, "Streaming", LOCTEXT("Cloner.Streaming", "Streaming"));
+	StreamingSection->RemoveCategory("World Partition");
+	StreamingSection->RemoveCategory("Data Layers");
+	StreamingSection->RemoveCategory("HLOD");
+
+	const TSharedRef<FPropertySection> GeneralSection = PropertyModule.FindOrCreateSection(ClassName, "General", LOCTEXT("Cloner.General", "General"));
+	GeneralSection->AddCategory("Cloner");
+
+	const TSharedRef<FPropertySection> LayoutSection = PropertyModule.FindOrCreateSection(ClassName, "Layouts", LOCTEXT("Cloner.Layouts", "Layouts"));
+	LayoutSection->AddCategory("Layout");
+
+	const TSharedRef<FPropertySection> RendererSection = PropertyModule.FindOrCreateSection(ClassName, "Renderer", LOCTEXT("Cloner.Renderer", "Renderer"));
+	RendererSection->AddCategory("Renderer");
+
+	const TSharedRef<FPropertySection> ProgressSection = PropertyModule.FindOrCreateSection(ClassName, "Progress", LOCTEXT("Cloner.Progress", "Progress"));
+	ProgressSection->AddCategory("Progress");
+
+	const TSharedRef<FPropertySection> StepSection = PropertyModule.FindOrCreateSection(ClassName, "Step", LOCTEXT("Cloner.Step", "Step"));
+	StepSection->AddCategory("Step");
+
+	const TSharedRef<FPropertySection> RangeSection = PropertyModule.FindOrCreateSection(ClassName, "Range", LOCTEXT("Cloner.Range", "Range"));
+	RangeSection->AddCategory("Range");
+
+	const TSharedRef<FPropertySection> SpawnSection = PropertyModule.FindOrCreateSection(ClassName, "Spawn", LOCTEXT("Cloner.Spawn", "Spawn"));
+	SpawnSection->AddCategory("Spawn");
+
+	const TSharedRef<FPropertySection> LifetimeSection = PropertyModule.FindOrCreateSection(ClassName, "Lifetime", LOCTEXT("Cloner.Lifetime", "Lifetime"));
+	LifetimeSection->AddCategory("Lifetime");
+}
+
+#undef LOCTEXT_NAMESPACE
diff --git a/Engine/Plugins/Experimental/ClonerEffector/Source/ClonerEffectorEditor/Private/Cloner/CEEditorClonerDetailCustomization.h b/Engine/Plugins/Experimental/ClonerEffector/Source/ClonerEffectorEditor/Private/Cloner/CEEditorClonerDetailCustomization.h
new file mode 100644
--- /dev/null
+++ b/Engine/Plugins/Experimental/ClonerEffector/Source/ClonerEffectorEditor/Private/Cloner/CEEditorClonerDetailCustomization.h
@@ -0,0 +1,37 @@
+// Copyright Epic Games, Inc. All Rights Reserved.
+
+#pragma once
+
+#include "IDetailCustomization.h"
+#include "Templates/SharedPointer.h"
+#include "UObject/Class.h"
+#include "UObject/Object.h"
+#include "UObject/WeakObjectPtr.h"
+
+class FReply;
+
+/* Used to customize cloner actor properties in details panel */
+class FCEEditorClonerDetailCustomization : public IDetailCustomization
+{
+public:
+	static TSharedRef<IDetailCustomization> MakeInstance()
+	{
+		return MakeShared<FCEEditorClonerDetailCustomization>();
+	}
+
+	explicit FCEEditorClonerDetailCustomization()
+	{
+		RegisterCustomSections();
+	}
+
+	virtual void CustomizeDetails(IDetailLayoutBuilder& InDetailBuilder) override;
+
+protected:
+	/** Execute ufunction with that name on selected objects */
+	FReply OnExecuteFunction(FName InFunctionName);
+
+	/** Register custom categories section in details panel */
+	void RegisterCustomSections() const;
+
+	TMap<FName, TMap<TWeakObjectPtr<UObject>, TWeakObjectPtr<UFunction>>> LayoutFunctionNames;
+};
diff --git a/Engine/Plugins/Experimental/ClonerEffector/Source/ClonerEffectorEditor/Private/Effector/CEEditorEffectorDetailCustomization.cpp b/Engine/Plugins/Experimental/ClonerEffector/Source/ClonerEffectorEditor/Private/Effector/CEEditorEffectorDetailCustomization.cpp
new file mode 100644
--- /dev/null
+++ b/Engine/Plugins/Experimental/ClonerEffector/Source/ClonerEffectorEditor/Private/Effector/CEEditorEffectorDetailCustomization.cpp
@@ -0,0 +1,250 @@
+// Copyright Epic Games, Inc. All Rights Reserved.
+
+#include "Effector/CEEditorEffectorDetailCustomization.h"
+
+#include "CEClonerEffectorShared.h"
+#include "DetailLayoutBuilder.h"
+#include "DetailWidgetRow.h"
+#include "Effector/CEEffectorActor.h"
+#include "Modules/ModuleManager.h"
+#include "PropertyEditorModule.h"
+#include "PropertyHandle.h"
+#include "Styles/CEEditorStyle.h"
+#include "Styling/SlateBrush.h"
+#include "Widgets/Images/SImage.h"
+#include "Widgets/Input/SComboBox.h"
+
+#define LOCTEXT_NAMESPACE "CEEditorEffectorDetailCustomization"
+
+DEFINE_LOG_CATEGORY_STATIC(LogCEEditorEffectorDetailCustomization, Log, All);
+
+void FCEEditorEffectorDetailCustomization::CustomizeDetails(IDetailLayoutBuilder& DetailBuilder)
+{
+	EasingPropertyHandle = DetailBuilder.GetProperty(
+		GET_MEMBER_NAME_CHECKED(ACEEffectorActor, Easing),
+		ACEEffectorActor::StaticClass()
+	);
+
+	if (!EasingPropertyHandle.IsValid())
+	{
+		return;
+	}
+
+	if (IDetailPropertyRow* EasingRow = DetailBuilder.EditDefaultProperty(EasingPropertyHandle))
+	{
+		PopulateEasingInfos();
+
+		FDetailWidgetRow& CustomWidget = EasingRow->CustomWidget();
+
+		CustomWidget.NameContent()
+		.HAlign(EHorizontalAlignment::HAlign_Fill)
+		.VAlign(EVerticalAlignment::VAlign_Fill)
+		[
+			EasingPropertyHandle->CreatePropertyNameWidget()
+		];
+
+		CustomWidget.ValueContent()
+		.HAlign(EHorizontalAlignment::HAlign_Fill)
+		.VAlign(EVerticalAlignment::VAlign_Fill)
+		[
+			SNew(SComboBox<FName>)
+			.OptionsSource(&EasingNames)
+			.InitiallySelectedItem(GetCurrentEasingName())
+			.ToolTipText(LOCTEXT("EasingTooltip", "Easings sorted from most dramatic to least and specials at the end"))
+			.OnGenerateWidget(this, &FCEEditorEffectorDetailCustomization::OnGenerateEasingEntry)
+			.OnSelectionChanged(this, &FCEEditorEffectorDetailCustomization::OnSelectionChanged)
+			.Content()
+			[
+				OnGenerateEasingEntry(NAME_None)
+			]
+		];
+	}
+}
+
+void FCEEditorEffectorDetailCustomization::RegisterCustomSections() const
+{
+	static const FName PropertyEditor("PropertyEditor");
+	FPropertyEditorModule& PropertyModule = FModuleManager::GetModuleChecked<FPropertyEditorModule>(PropertyEditor);
+
+	static const FName ClassName = ACEEffectorActor::StaticClass()->GetFName();
+
+	const TSharedRef<FPropertySection> StreamingSection = PropertyModule.FindOrCreateSection(ClassName, "Streaming", LOCTEXT("Effector.Streaming", "Streaming"));
+	StreamingSection->RemoveCategory("World Partition");
+	StreamingSection->RemoveCategory("Data Layers");
+	StreamingSection->RemoveCategory("HLOD");
+
+	const TSharedRef<FPropertySection> EffectorSection = PropertyModule.FindOrCreateSection(ClassName, "General", LOCTEXT("Effector.General", "General"));
+	EffectorSection->AddCategory(TEXT("Effector"));
+
+	const TSharedRef<FPropertySection> TypeSection = PropertyModule.FindOrCreateSection(ClassName, "Type", LOCTEXT("Effector.Type", "Type"));
+	TypeSection->AddCategory(TEXT("Type"));
+
+	const TSharedRef<FPropertySection> ModeSection = PropertyModule.FindOrCreateSection(ClassName, "Mode", LOCTEXT("Effector.Mode", "Mode"));
+	ModeSection->AddCategory(TEXT("Mode"));
+
+	const TSharedRef<FPropertySection> ForceSection = PropertyModule.FindOrCreateSection(ClassName, "Forces", LOCTEXT("Effector.Forces", "Forces"));
+	ForceSection->AddCategory(TEXT("Force"));
+}
+
+void FCEEditorEffectorDetailCustomization::PopulateEasingInfos()
+{
+	EasingEnum = StaticEnum<ECEClonerEasing>();
+	if (EasingEnum.IsValid())
+	{
+		EasingNames.Empty();
+
+		// Sort from most dramatic to least IN then OUT then IN OUT, then specials
+		static const TArray<ECEClonerEasing> SortedEasings
+		{
+			ECEClonerEasing::InExpo,
+			ECEClonerEasing::InCirc,
+			ECEClonerEasing::InQuint,
+			ECEClonerEasing::InQuart,
+			ECEClonerEasing::InQuad,
+			ECEClonerEasing::InCubic,
+			ECEClonerEasing::InSine,
+			ECEClonerEasing::OutExpo,
+			ECEClonerEasing::OutCirc,
+			ECEClonerEasing::OutQuint,
+			ECEClonerEasing::OutQuart,
+			ECEClonerEasing::OutQuad,
+			ECEClonerEasing::OutCubic,
+			ECEClonerEasing::OutSine,
+			ECEClonerEasing::InOutExpo,
+			ECEClonerEasing::InOutCirc,
+			ECEClonerEasing::InOutQuint,
+			ECEClonerEasing::InOutQuart,
+			ECEClonerEasing::InOutQuad,
+			ECEClonerEasing::InOutCubic,
+			ECEClonerEasing::InOutSine,
+			ECEClonerEasing::Linear,
+			ECEClonerEasing::InBounce,
+			ECEClonerEasing::InBack,
+			ECEClonerEasing::InElastic,
+			ECEClonerEasing::OutBounce,
+			ECEClonerEasing::OutBack,
+			ECEClonerEasing::OutElastic,
+			ECEClonerEasing::InOutBounce,
+			ECEClonerEasing::InOutBack,
+			ECEClonerEasing::InOutElastic,
+			ECEClonerEasing::Random
+		};
+
+		check(EasingEnum->GetMaxEnumValue() == SortedEasings.Num());
+
+		for (const ECEClonerEasing& Easing : SortedEasings)
+		{
+			EasingNames.Add(FName(EasingEnum->GetNameStringByValue(static_cast<uint8>(Easing))));
+		}
+	}
+}
+
+TSharedRef<SWidget> FCEEditorEffectorDetailCustomization::OnGenerateEasingEntry(FName InName) const
+{
+	TSharedPtr<SWidget> ImageWidget;
+	TSharedPtr<SWidget> TextWidget;
+
+	static const FVector2D ImageSizeClosed(16.f, 16.f);
+	static const FVector2D ImageSizeOpened(32.f, 32.f);
+
+	// If none = update with current value else set fixed value once
+	if (InName == NAME_None)
+	{
+		SAssignNew(ImageWidget, SImage)
+		.DesiredSizeOverride(ImageSizeClosed)
+		.Image(this, &FCEEditorEffectorDetailCustomization::GetEasingImage, InName);
+
+		SAssignNew(TextWidget, STextBlock)
+		.Text(this, &FCEEditorEffectorDetailCustomization::GetEasingText, InName);
+	}
+	else
+	{
+		SAssignNew(ImageWidget, SImage)
+		.DesiredSizeOverride(ImageSizeOpened)
+		.Image(GetEasingImage(InName));
+
+		SAssignNew(TextWidget, STextBlock)
+		.Text(GetEasingText(InName));
+	}
+
+	return SNew(SHorizontalBox)
+		+ SHorizontalBox::Slot()
+		.AutoWidth()
+		.Padding(2.f)
+		[
+			ImageWidget.ToSharedRef()
+		]
+		+ SHorizontalBox::Slot()
+		.FillWidth(1.f)
+		.Padding(5.f, 2.f)
+		.HAlign(EHorizontalAlignment::HAlign_Left)
+		.VAlign(EVerticalAlignment::VAlign_Center)
+		[
+			TextWidget.ToSharedRef()
+		];
+}
+
+void FCEEditorEffectorDetailCustomization::OnSelectionChanged(FName InSelection, ESelectInfo::Type InSelectInfo) const
+{
+	if (EasingEnum.IsValid() && EasingPropertyHandle.IsValid())
+	{
+		const int32 Idx = EasingEnum->GetValueByNameString(InSelection.ToString());
+		if (Idx != INDEX_NONE && EasingPropertyHandle->SetValue(static_cast<uint8>(Idx), EPropertyValueSetFlags::DefaultFlags) != FPropertyAccess::Success)
+		{
+			UE_LOG(LogCEEditorEffectorDetailCustomization, Warning, TEXT("ClonerEffectorDetailCustomization : Cannot set property value %s on selection"), *InSelection.ToString())
+		}
+	}
+}
+
+FName FCEEditorEffectorDetailCustomization::GetCurrentEasingName() const
+{
+	uint8 CurrentValue;
+	if (EasingPropertyHandle.IsValid() && EasingPropertyHandle->GetValue(CurrentValue) == FPropertyAccess::Result::Success)
+	{
+		if (EasingEnum.IsValid())
+		{
+			return FName(EasingEnum->GetNameStringByValue(CurrentValue));
+		}
+	}
+	return NAME_None;
+}
+
+const FSlateBrush* FCEEditorEffectorDetailCustomization::GetEasingImage(FName InName) const
+{
+	if (InName == NAME_None)
+	{
+		InName = GetCurrentEasingName();
+
+		// multiple values
+		if (InName == NAME_None)
+		{
+			return nullptr;
+		}
+	}
+
+	return FCEEditorStyle::Get().GetBrush(FName(TEXT("EasingIcons.") + InName.ToString()));
+}
+
+FText FCEEditorEffectorDetailCustomization::GetEasingText(FName InName) const
+{
+	if (InName == NAME_None)
+	{
+		InName = GetCurrentEasingName();
+
+		// multiple values
+		if (InName == NAME_None)
+		{
+			return LOCTEXT("MultipleValue", "Multiple values selected");
+		}
+	}
+
+	if (EasingEnum.IsValid())
+	{
+		const int32 EnumValue = EasingEnum->GetValueByNameString(InName.ToString());
+		return EasingEnum->GetDisplayNameTextByValue(EnumValue);
+	}
+
+	return FText::GetEmpty();
+}
+
+#undef LOCTEXT_NAMESPACE
diff --git a/Engine/Plugins/Experimental/ClonerEffector/Source/ClonerEffectorEditor/Private/Effector/CEEditorEffectorDetailCustomization.h b/Engine/Plugins/Experimental/ClonerEffector/Source/ClonerEffectorEditor/Private/Effector/CEEditorEffectorDetailCustomization.h
new file mode 100644
--- /dev/null
+++ b/Engine/Plugins/Experimental/ClonerEffector/Source/ClonerEffectorEditor/Private/Effector/CEEditorEffectorDetailCustomization.h
@@ -0,0 +1,47 @@
+// Copyright Epic Games, Inc. All Rights Reserved.
+
+#pragma once
+
+#include "IDetailCustomization.h"
+#include "Templates/SharedPointer.h"
+#include "Types/SlateEnums.h"
+#include "UObject/WeakObjectPtr.h"
+
+class IPropertyHandle;
+class SWidget;
+class UEnum;
+struct FSlateBrush;
+
+/* Used to customize effector actor properties in details panel */
+class FCEEditorEffectorDetailCustomization : public IDetailCustomization
+{
+public:
+	static TSharedRef<IDetailCustomization> MakeInstance()
+	{
+		return MakeShared<FCEEditorEffectorDetailCustomization>();
+	}
+
+	explicit FCEEditorEffectorDetailCustomization()
+	{
+		RegisterCustomSections();
+	}
+
+	//~ Begin IDetailCustomization
+	virtual void CustomizeDetails(IDetailLayoutBuilder& DetailBuilder) override;
+	//~ End IDetailCustomization
+
+protected:
+	void RegisterCustomSections() const;
+
+	void PopulateEasingInfos();
+	FName GetCurrentEasingName() const;
+	const FSlateBrush* GetEasingImage(FName InName) const;
+	FText GetEasingText(FName InName) const;
+
+	TSharedRef<SWidget> OnGenerateEasingEntry(FName InName) const;
+	void OnSelectionChanged(FName InSelection, ESelectInfo::Type InSelectInfo) const;
+
+	TWeakObjectPtr<UEnum> EasingEnum;
+	TArray<FName> EasingNames;
+	TSharedPtr<IPropertyHandle> EasingPropertyHandle;
+};
