MotionDesign
- Change ClonerEffector filetype from utf8 to text

#changelist validated
#virtualized

diff --git a/Engine/Plugins/Experimental/ClonerEffector/Source/ClonerEffector/ClonerEffector.Build.cs b/Engine/Plugins/Experimental/ClonerEffector/Source/ClonerEffector/ClonerEffector.Build.cs
--- a/Engine/Plugins/Experimental/ClonerEffector/Source/ClonerEffector/ClonerEffector.Build.cs
+++ b/Engine/Plugins/Experimental/ClonerEffector/Source/ClonerEffector/ClonerEffector.Build.cs
@@ -1,4 +1,4 @@
-﻿// Copyright Epic Games, Inc. All Rights Reserved.
+// Copyright Epic Games, Inc. All Rights Reserved.
 
 using UnrealBuildTool;
 
diff --git a/Engine/Plugins/Experimental/ClonerEffector/Source/ClonerEffector/Private/CEModule.cpp b/Engine/Plugins/Experimental/ClonerEffector/Source/ClonerEffector/Private/CEModule.cpp
--- a/Engine/Plugins/Experimental/ClonerEffector/Source/ClonerEffector/Private/CEModule.cpp
+++ b/Engine/Plugins/Experimental/ClonerEffector/Source/ClonerEffector/Private/CEModule.cpp
@@ -1,4 +1,4 @@
-﻿// Copyright Epic Games, Inc. All Rights Reserved.
+// Copyright Epic Games, Inc. All Rights Reserved.
 
 #include "Modules/ModuleManager.h"
 
diff --git a/Engine/Plugins/Experimental/ClonerEffector/Source/ClonerEffectorEditor/ClonerEffectorEditor.Build.cs b/Engine/Plugins/Experimental/ClonerEffector/Source/ClonerEffectorEditor/ClonerEffectorEditor.Build.cs
--- a/Engine/Plugins/Experimental/ClonerEffector/Source/ClonerEffectorEditor/ClonerEffectorEditor.Build.cs
+++ b/Engine/Plugins/Experimental/ClonerEffector/Source/ClonerEffectorEditor/ClonerEffectorEditor.Build.cs
@@ -1,4 +1,4 @@
-﻿// Copyright Epic Games, Inc. All Rights Reserved.
+// Copyright Epic Games, Inc. All Rights Reserved.
 
 using UnrealBuildTool;
 
diff --git a/Engine/Plugins/Experimental/ClonerEffector/Source/ClonerEffectorEditor/Private/CEEditorModule.cpp b/Engine/Plugins/Experimental/ClonerEffector/Source/ClonerEffectorEditor/Private/CEEditorModule.cpp
--- a/Engine/Plugins/Experimental/ClonerEffector/Source/ClonerEffectorEditor/Private/CEEditorModule.cpp
+++ b/Engine/Plugins/Experimental/ClonerEffector/Source/ClonerEffectorEditor/Private/CEEditorModule.cpp
@@ -1,4 +1,4 @@
-﻿// Copyright Epic Games, Inc. All Rights Reserved.
+// Copyright Epic Games, Inc. All Rights Reserved.
 
 #include "CEEditorModule.h"
 
diff --git a/Engine/Plugins/Experimental/ClonerEffector/Source/ClonerEffectorEditor/Private/CEEditorModule.h b/Engine/Plugins/Experimental/ClonerEffector/Source/ClonerEffectorEditor/Private/CEEditorModule.h
--- a/Engine/Plugins/Experimental/ClonerEffector/Source/ClonerEffectorEditor/Private/CEEditorModule.h
+++ b/Engine/Plugins/Experimental/ClonerEffector/Source/ClonerEffectorEditor/Private/CEEditorModule.h
@@ -1,4 +1,4 @@
-﻿// Copyright Epic Games, Inc. All Rights Reserved.
+// Copyright Epic Games, Inc. All Rights Reserved.
 
 #pragma once
 
