MotionDesign : ClonerEffector
- Fix simulation behavior outside of editor (MotionDesign Rundown)
- Improvements to the structure by simplifying usage, link effectors within cloner instead of linking cloners within effector
- Replacing deprecated Vortex module by new Vortex Force module in assets

#changelist validated
#virtualized

diff --git a/Engine/Plugins/Experimental/ClonerEffector/Source/ClonerEffector/Private/CEClonerEffectorShared.cpp b/Engine/Plugins/Experimental/ClonerEffector/Source/ClonerEffector/Private/CEClonerEffectorShared.cpp
--- a/Engine/Plugins/Experimental/ClonerEffector/Source/ClonerEffector/Private/CEClonerEffectorShared.cpp
+++ b/Engine/Plugins/Experimental/ClonerEffector/Source/ClonerEffector/Private/CEClonerEffectorShared.cpp
@@ -138,6 +138,14 @@ FCEClonerEffectorDataInterfaces::FCEClonerEffectorDataInterfaces(const UNiagaraS
 	DataInterfaces.Add(EffectorIndexDIVar.GetName(), IndexArrayDI);
 }
 
+void FCEClonerEffectorDataInterfaces::Clear() const
+{
+	if (UNiagaraDataInterfaceArrayInt32* IndexArray = GetIndexArray())
+	{
+		IndexArray->GetArrayReference().Empty();
+	}
+}
+
 void FCEClonerEffectorDataInterfaces::CopyTo(FCEClonerEffectorDataInterfaces& InOther) const
 {
 	for (const TPair<FName, TObjectPtr<UNiagaraDataInterface>>& DataInterfacePair : DataInterfaces)
diff --git a/Engine/Plugins/Experimental/ClonerEffector/Source/ClonerEffector/Private/Cloner/CEClonerActor.cpp b/Engine/Plugins/Experimental/ClonerEffector/Source/ClonerEffector/Private/Cloner/CEClonerActor.cpp
--- a/Engine/Plugins/Experimental/ClonerEffector/Source/ClonerEffector/Private/Cloner/CEClonerActor.cpp
+++ b/Engine/Plugins/Experimental/ClonerEffector/Source/ClonerEffector/Private/Cloner/CEClonerActor.cpp
@@ -18,6 +18,7 @@
 #include "NiagaraMeshRendererProperties.h"
 #include "NiagaraSystem.h"
 #include "Subsystems/CEClonerSubsystem.h"
+#include "Subsystems/CEEffectorSubsystem.h"
 #include "UObject/ConstructorHelpers.h"
 
 #if WITH_EDITOR
@@ -35,6 +36,7 @@ ACEClonerActor::ACEClonerActor()
 	PrimaryActorTick.bCanEverTick          = true;
 	PrimaryActorTick.bStartWithTickEnabled = true;
 	PrimaryActorTick.bTickEvenWhenPaused   = true;
+	PrimaryActorTick.bHighPriority         = true;
 
 	ClonerComponent = CreateDefaultSubobject<UCEClonerComponent>(TEXT("AvaClonerComponent"));
 	SetRootComponent(ClonerComponent);
@@ -52,9 +54,11 @@ ACEClonerActor::ACEClonerActor()
 		}
 #endif
 
-		ClonerComponent->TransformUpdated.AddUObject(this, &ACEClonerActor::OnClonerTransformed);
 		ClonerComponent->OnClonerMeshUpdated.AddUObject(this, &ACEClonerActor::OnClonerMeshUpdated);
 
+		UCEEffectorSubsystem::OnEffectorIdentifierChangedDelegate.AddUObject(this, &ACEClonerActor::OnEffectorIdentifierChanged);
+		ACEEffectorActor::OnEffectorRefreshClonerDelegate.AddUObject(this, &ACEClonerActor::OnEffectorRefreshCloner);
+
 		const TArray<FString> LayoutNames = GetClonerLayoutNames();
 
 		// Apply default layout
@@ -80,6 +84,7 @@ const TCEPropertyChangeDispatcher<ACEClonerActor> ACEClonerActor::PropertyChange
 {
 	{ GET_MEMBER_NAME_CHECKED(ACEClonerActor, bEnabled), &ACEClonerActor::OnEnabledChanged },
 	{ GET_MEMBER_NAME_CHECKED(ACEClonerActor, Seed), &ACEClonerActor::OnSeedChanged },
+	{ GET_MEMBER_NAME_CHECKED(ACEClonerActor, EffectorsWeak), &ACEClonerActor::OnEffectorsChanged },
 	/** Layout */
 	{ GET_MEMBER_NAME_CHECKED(ACEClonerActor, LayoutName), &ACEClonerActor::OnLayoutNameChanged },
 	/** Advanced */
@@ -295,7 +300,7 @@ void ACEClonerActor::PostActorCreated()
 	Super::PostActorCreated();
 
 #if WITH_EDITOR
-	SpawnDefaultActorAttached();
+	bSpawnDefaultActorAttached = true;
 #endif
 }
 
@@ -351,33 +356,6 @@ void ACEClonerActor::UpdateLayoutOptions()
 	OnLifetimeOptionsChanged();
 }
 
-void ACEClonerActor::UpdateClonerEffectors()
-{
-	Effectors.RemoveAll([](const TWeakObjectPtr<ACEEffectorActor>& InEffector)
-	{
-		return !InEffector.IsValid();
-	});
-
-	if (const UCEClonerLayoutBase* LayoutSystem = GetActiveLayout())
-	{
-		// Effectors could be registered before system is loaded
-		const int32 EffectorCount = GetEffectorCount();
-
-		if (LayoutSystem->GetDataInterfaces().Num() != EffectorCount)
-		{
-			LayoutSystem->GetDataInterfaces().Resize(EffectorCount);
-
-			RequestClonerUpdate();
-		}
-	}
-
-	ForEachEffector([this](ACEEffectorActor* InEffector, int32 InIdx)
-	{
-		InEffector->OnClonerUpdated(this, InIdx);
-		return true;
-	});
-}
-
 void ACEClonerActor::SetTreeUpdateInterval(float InInterval)
 {
 	if (InInterval == TreeUpdateInterval)
@@ -882,85 +860,54 @@ int32 ACEClonerActor::GetMeshCount() const
 	return 0;
 }
 
-int32 ACEClonerActor::RegisterEffector(ACEEffectorActor* InEffector)
+bool ACEClonerActor::LinkEffector(ACEEffectorActor* InEffector)
 {
-	if (!InEffector)
+	if (!IsValid(InEffector) || EffectorsWeak.Contains(InEffector))
 	{
-		return INDEX_NONE;
+		return false;
 	}
 
-	int32 Index = Effectors.Find(InEffector);
+	EffectorsWeak.Add(InEffector);
 
-	if (Index != INDEX_NONE)
-	{
-		return Index;
-	}
+	OnEffectorsChanged();
 
-	Index = Effectors.Add(InEffector);
+	UE_LOG(LogCEClonerActor, Log, TEXT("%s : Effector %s linked to Cloner"), *GetActorNameOrLabel(), *InEffector->GetActorNameOrLabel());
 
-	if (const UCEClonerLayoutBase* LayoutSystem = GetActiveLayout())
-	{
-		const int32 EffectorCount = GetEffectorCount();
-		LayoutSystem->GetDataInterfaces().Resize(EffectorCount);
-	}
-
-	InEffector->OnClonerLinked(this, Index);
-
-	constexpr bool bImmediateUpdate = true;
-	RequestClonerUpdate(bImmediateUpdate);
-
-	return Index;
+	return true;
 }
 
-bool ACEClonerActor::UnregisterEffector(ACEEffectorActor* InEffector)
+bool ACEClonerActor::UnlinkEffector(ACEEffectorActor* InEffector)
 {
 	if (!InEffector)
 	{
 		return false;
 	}
 
-	const int32 OldIndex = Effectors.Find(InEffector);
-
-	if (OldIndex == INDEX_NONE)
+	if (EffectorsWeak.Remove(InEffector) > 0)
 	{
-		return false;
-	}
-
-	Effectors.RemoveAt(OldIndex);
+		OnEffectorsChanged();
 
-	if (const UCEClonerLayoutBase* LayoutSystem = GetActiveLayout())
-	{
-		LayoutSystem->GetDataInterfaces().Remove(OldIndex);
+		UE_LOG(LogCEClonerActor, Log, TEXT("%s : Effector %s unlinked from Cloner"), *GetActorNameOrLabel(), *InEffector->GetActorNameOrLabel());
 	}
 
-	InEffector->OnClonerUnlinked(this, OldIndex);
-
-	constexpr bool bImmediateUpdate = true;
-	RequestClonerUpdate(bImmediateUpdate);
-
 	return true;
 }
 
-bool ACEClonerActor::IsEffectorRegistered(const ACEEffectorActor* InEffector) const
-{
-	return InEffector && Effectors.Contains(InEffector);
-}
-
-int32 ACEClonerActor::GetEffectorIndex(ACEEffectorActor* InEffector) const
+bool ACEClonerActor::IsEffectorLinked(const ACEEffectorActor* InEffector) const
 {
-	return InEffector ? Effectors.Find(InEffector) : INDEX_NONE;
+	return InEffector && EffectorsWeak.Contains(InEffector);
 }
 
 int32 ACEClonerActor::GetEffectorCount() const
 {
-	return Effectors.Num();
+	return EffectorsWeak.Num();
 }
 
 void ACEClonerActor::ForEachEffector(TFunctionRef<bool(ACEEffectorActor*, int32)> InFunction)
 {
-	for (int32 Idx = 0; Idx < Effectors.Num(); Idx++)
+	for (int32 Idx = 0; Idx < EffectorsWeak.Num(); Idx++)
 	{
-		ACEEffectorActor* Effector = Effectors[Idx].Get();
+		ACEEffectorActor* Effector = EffectorsWeak[Idx].Get();
 		if (!Effector)
 		{
 			continue;
@@ -1085,20 +1032,74 @@ void ACEClonerActor::OnLifetimeOptionsChanged()
 	RequestClonerUpdate();
 }
 
-void ACEClonerActor::OnClonerTransformed(USceneComponent*, EUpdateTransformFlags, ETeleportType)
+void ACEClonerActor::OnClonerMeshUpdated(UCEClonerComponent* InClonerComponent)
 {
-	UpdateClonerEffectors();
+	if (InClonerComponent == ClonerComponent)
+	{
+		RequestClonerUpdate();
+	}
 }
 
-void ACEClonerActor::OnClonerMeshUpdated()
+void ACEClonerActor::OnClonerSystemChanged()
 {
-	RequestClonerUpdate();
+	UpdateLayoutOptions();
 }
 
-void ACEClonerActor::OnClonerSystemChanged()
+void ACEClonerActor::OnEffectorIdentifierChanged(ACEEffectorActor* InEffector)
 {
-	UpdateLayoutOptions();
-	UpdateClonerEffectors();
+	if (EffectorsWeak.Contains(InEffector))
+	{
+		OnEffectorsChanged();
+	}
+}
+
+void ACEClonerActor::OnEffectorRefreshCloner(ACEEffectorActor* InEffector)
+{
+	if (EffectorsWeak.Contains(InEffector))
+	{
+		RequestClonerUpdate();
+	}
+}
+
+void ACEClonerActor::OnEffectorsChanged()
+{
+	const FCEClonerEffectorDataInterfaces* EffectorDataInterfaces = GetEffectorDataInterfaces();
+
+	if (!EffectorDataInterfaces || !EffectorDataInterfaces->GetIndexArray())
+	{
+		return;
+	}
+
+	// Remove duplicates
+	const TSet<TWeakObjectPtr<ACEEffectorActor>> SetEffectorsWeak(EffectorsWeak);
+	EffectorsWeak = SetEffectorsWeak.Array();
+
+	TArray<int32> EffectorIndexes;
+	EffectorIndexes.Reserve(EffectorsWeak.Num());
+
+	for (const TWeakObjectPtr<ACEEffectorActor>& EffectorWeak : EffectorsWeak)
+	{
+		if (const ACEEffectorActor* Effector = EffectorWeak.Get())
+		{
+			const int32 ChannelIdentifier = Effector->GetChannelIdentifier();
+
+			if (ChannelIdentifier != INDEX_NONE)
+			{
+				EffectorIndexes.AddUnique(ChannelIdentifier);
+			}
+		}
+	}
+
+	EffectorDataInterfaces->Clear();
+	TArray<int32>& EffectorIndexArray = EffectorDataInterfaces->GetIndexArray()->GetArrayReference();
+
+	for (const int32 EffectorIndex : EffectorIndexes)
+	{
+		EffectorIndexArray.Add(EffectorIndex);
+	}
+
+	constexpr bool bImmediateUpdate = false;
+	RequestClonerUpdate(bImmediateUpdate);
 }
 
 void ACEClonerActor::RequestClonerUpdate(bool bInImmediate)
@@ -1248,15 +1249,23 @@ void ACEClonerActor::SpawnLinkedEffector()
 		return;
 	}
 
-	EffectorActor->LinkCloner(this);
 	FActorLabelUtilities::RenameExistingActor(EffectorActor, EffectorActor->GetDefaultActorLabel(), true);
 
 	// Set default offset for visual feedback
 	EffectorActor->SetOffset(FVector(0, 0, 100));
+
+	LinkEffector(EffectorActor);
 }
 
 void ACEClonerActor::SpawnDefaultActorAttached()
 {
+	if (!bSpawnDefaultActorAttached)
+	{
+		return;
+	}
+
+	bSpawnDefaultActorAttached = false;
+
 	// Only spawn if world is valid and not a preview actor
 	UWorld* World = GetWorld();
 	if (!World || bIsEditorPreviewActor)
@@ -1264,6 +1273,17 @@ void ACEClonerActor::SpawnDefaultActorAttached()
 		return;
 	}
 
+	// Only spawn if no actor is attached below it
+	TArray<AActor*> AttachedActors;
+	constexpr bool bReset = true;
+	constexpr bool bRecursive = false;
+	GetAttachedActors(AttachedActors, bReset, bRecursive);
+
+	if (!AttachedActors.IsEmpty())
+	{
+		return;
+	}
+
 	// Find or load cube mesh
 	constexpr const TCHAR* DefaultStaticMeshPath = TEXT("/Script/Engine.StaticMesh'/Engine/BasicShapes/Cube.Cube'");
 	UStaticMesh* DefaultStaticMesh = LoadObject<UStaticMesh>(nullptr, DefaultStaticMeshPath);
@@ -1393,6 +1413,11 @@ void ACEClonerActor::InitializeCloner()
 
 #if WITH_EDITOR
 	OnReduceMotionGhostingChanged();
+
+	if (bSpawnDefaultActorAttached)
+	{
+		SpawnDefaultActorAttached();
+	}
 #endif
 
 	// For new cloner instances no need to migrate anything
@@ -1401,4 +1426,5 @@ void ACEClonerActor::InitializeCloner()
 	bClonerInitialized = true;
 
 	OnLayoutNameChanged();
+	OnEffectorsChanged();
 }
diff --git a/Engine/Plugins/Experimental/ClonerEffector/Source/ClonerEffector/Private/Cloner/CEClonerComponent.cpp b/Engine/Plugins/Experimental/ClonerEffector/Source/ClonerEffector/Private/Cloner/CEClonerComponent.cpp
--- a/Engine/Plugins/Experimental/ClonerEffector/Source/ClonerEffector/Private/Cloner/CEClonerComponent.cpp
+++ b/Engine/Plugins/Experimental/ClonerEffector/Source/ClonerEffector/Private/Cloner/CEClonerComponent.cpp
@@ -43,6 +43,8 @@ UCEClonerComponent::UCEClonerComponent()
 {
 	CastShadow = true;
 	bReceivesDecals = true;
+	bAutoActivate = true;
+	bHiddenInGame = false;
 
 #if WITH_EDITOR
 	// Do not show bounding box around cloner for better visibility
@@ -1341,7 +1343,7 @@ void UCEClonerComponent::UpdateClonerMeshes()
 
 	UE_LOG(LogCEClonerComponent, Log, TEXT("%s : Cloner mesh updated %i"), *ClonerActor->GetActorNameOrLabel(), ClonerTree.MergedBakedMeshes.Num());
 
-	OnClonerMeshUpdated.Broadcast();
+	OnClonerMeshUpdated.Broadcast(this);
 }
 
 bool UCEClonerComponent::SetClonerActiveLayout(UCEClonerLayoutBase* InLayout)
@@ -1368,7 +1370,7 @@ bool UCEClonerComponent::SetClonerActiveLayout(UCEClonerLayoutBase* InLayout)
 
 		UE_LOG(LogCEClonerComponent, Log, TEXT("%s : Cloner layout system loaded %s - %s"), *GetOwner()->GetActorNameOrLabel(), *InLayout->GetLayoutName().ToString(), *InLayout->GetLayoutAssetPath());
 
-		OnClonerSystemLoaded.Broadcast(InLayout);
+		OnClonerSystemLoaded.Broadcast(this, InLayout);
 	}
 
 	// Copy data interfaces to new layout
diff --git a/Engine/Plugins/Experimental/ClonerEffector/Source/ClonerEffector/Private/Effector/CEEffectorActor.cpp b/Engine/Plugins/Experimental/ClonerEffector/Source/ClonerEffector/Private/Effector/CEEffectorActor.cpp
--- a/Engine/Plugins/Experimental/ClonerEffector/Source/ClonerEffector/Private/Effector/CEEffectorActor.cpp
+++ b/Engine/Plugins/Experimental/ClonerEffector/Source/ClonerEffector/Private/Effector/CEEffectorActor.cpp
@@ -11,6 +11,8 @@
 
 DEFINE_LOG_CATEGORY_STATIC(LogCEEffectorActor, Log, All);
 
+ACEEffectorActor::FOnEffectorIdentifierChanged ACEEffectorActor::OnEffectorRefreshClonerDelegate;
+
 ACEEffectorActor::ACEEffectorActor()
 {
 	SetCanBeDamaged(false);
@@ -100,8 +102,6 @@ ACEEffectorActor::ACEEffectorActor()
 	if (!IsTemplate())
 	{
 		SceneComponent->TransformUpdated.AddUObject(this, &ACEEffectorActor::OnEffectorTransformed);
-
-		RegisterToChannel();
 	}
 }
 
@@ -116,7 +116,6 @@ TCEPropertyChangeDispatcher<ACEEffectorActor> ACEEffectorActor::PropertyChangeDi
 	/** Effector */
 	{ GET_MEMBER_NAME_CHECKED(ACEEffectorActor, bEnabled), &ACEEffectorActor::OnEnabledChanged },
 	{ GET_MEMBER_NAME_CHECKED(ACEEffectorActor, Magnitude), &ACEEffectorActor::OnMagnitudeChanged },
-	{ GET_MEMBER_NAME_CHECKED(ACEEffectorActor, Cloners), &ACEEffectorActor::OnClonersChanged },
 	{ GET_MEMBER_NAME_CHECKED(ACEEffectorActor, VisualizerThickness), &ACEEffectorActor::OnVisualizerThicknessChanged },
 	{ GET_MEMBER_NAME_CHECKED(ACEEffectorActor, bVisualizerSpriteVisible), &ACEEffectorActor::OnVisualizerSpriteVisibleChanged },
 	/** Type */
@@ -167,23 +166,14 @@ void ACEEffectorActor::PostEditImport()
 {
 	Super::PostEditImport();
 
-	RegisterToCloners();
+	OnEffectorChanged();
 }
 
 void ACEEffectorActor::PostEditUndo()
 {
 	Super::PostEditUndo();
 
-	RegisterToChannel();
-
-	OnClonersChanged();
-
-	// Undo causes cloner to blackout, so refresh state
-	ForEachCloner([](ACEClonerActor* InCloner, int32 InIndex)
-	{
-		InCloner->ForceUpdateCloner();
-		return true;
-	});
+	OnEffectorChanged();
 }
 #endif
 
@@ -196,13 +186,6 @@ void ACEEffectorActor::Destroyed()
 	{
 		EffectorSubsystem->UnregisterChannelEffector(this);
 	}
-
-	// Remove this effector from the cloners it is linked to
-	ForEachCloner([this](ACEClonerActor* InCloner, int32 InIndex)
-	{
-		InCloner->UnregisterEffector(this);
-		return true;
-	});
 }
 
 void ACEEffectorActor::PostActorCreated()
@@ -216,73 +199,50 @@ void ACEEffectorActor::PostLoad()
 {
 	Super::PostLoad();
 
-	RegisterToCloners();
-}
-
-void ACEEffectorActor::PostDuplicate(EDuplicateMode::Type InDuplicateMode)
-{
-	Super::PostDuplicate(InDuplicateMode);
-
-	RegisterToCloners();
-}
-
-void ACEEffectorActor::LinkCloner(ACEClonerActor* InCloner)
-{
-	if (!InCloner || Cloners.Contains(InCloner))
+	if (!InternalCloners.IsEmpty())
 	{
-		return;
+		for (const TWeakObjectPtr<ACEClonerActor>& ClonerWeak : InternalCloners)
+		{
+			if (ACEClonerActor* Cloner = ClonerWeak.Get())
+			{
+				Cloner->LinkEffector(this);
+			}
+		}
+
+		InternalCloners.Empty();
 	}
 
-	Cloners.Add(InCloner);
-	OnClonersChanged();
+	OnEffectorChanged();
 }
 
-void ACEEffectorActor::UnlinkCloner(ACEClonerActor* InCloner)
+void ACEEffectorActor::PostDuplicate(EDuplicateMode::Type InDuplicateMode)
 {
-	if (!InCloner || !Cloners.Contains(InCloner))
-	{
-		return;
-	}
+	Super::PostDuplicate(InDuplicateMode);
 
-	Cloners.Remove(InCloner);
-	OnClonersChanged();
+	OnEffectorChanged();
 }
 
-void ACEEffectorActor::OnEffectorIdentifierChanged()
+int32 ACEEffectorActor::GetChannelIdentifier() const
 {
-	int32 Identifier = ChannelData.GetIdentifier();
+	const int32 Identifier = ChannelData.GetIdentifier();
 
 	if (Identifier == INDEX_NONE)
 	{
-		UE_LOG(LogCEEffectorActor, Log, TEXT("%s : Effector channel identifier is invalid"), *GetActorNameOrLabel());
-		return;
-	}
+		UCEEffectorSubsystem* EffectorSubsystem = !!GetWorld() ? UCEEffectorSubsystem::Get(GetWorld()) : nullptr;
 
-	// Update effector identifier in linked cloners
-	ForEachCloner([Identifier](ACEClonerActor* InCloner, int32 InEffectorIndex)
-	{
-		if (const FCEClonerEffectorDataInterfaces* DataInterfaces = InCloner->GetEffectorDataInterfaces())
+		// Register this effector to the effector channel
+		if (EffectorSubsystem)
 		{
-			UNiagaraDataInterfaceArrayInt32* EffectorIndexDI = DataInterfaces->GetIndexArray();
-
-			// This effector has a specific index in the cloner
-			if (EffectorIndexDI->GetArrayReference().IsValidIndex(InEffectorIndex))
-			{
-				const int32 PrevIdentifier = EffectorIndexDI->GetArrayReference()[InEffectorIndex];
-
-				// Only update if identifier is different
-				if (PrevIdentifier != Identifier)
-				{
-					EffectorIndexDI->GetArrayReference()[InEffectorIndex] = Identifier;
-
-					constexpr bool bImmediateUpdate = true;
-					InCloner->RequestClonerUpdate(bImmediateUpdate);
-				}
-			}
+			EffectorSubsystem->RegisterChannelEffector(const_cast<ACEEffectorActor*>(this));
 		}
+	}
 
-		return true;
-	});
+	return ChannelData.GetIdentifier();
+}
+
+FCEClonerEffectorChannelData& ACEEffectorActor::GetChannelData()
+{
+	return ChannelData;
 }
 
 void ACEEffectorActor::SetType(ECEClonerEffectorType InType)
@@ -740,52 +700,6 @@ void ACEEffectorActor::SetVisualizerSpriteVisible(bool bInVisible)
 }
 #endif
 
-void ACEEffectorActor::ForEachCloner(TFunctionRef<bool(ACEClonerActor*, int32)> InFunction, bool bSkipIdxCheck)
-{
-	for (TSet<TWeakObjectPtr<ACEClonerActor>>::TIterator It(InternalCloners); It; ++It)
-	{
-		ACEClonerActor* Cloner = It->Get();
-		if (!Cloner)
-		{
-			continue;
-		}
-
-		const int32 Idx = Cloner->GetEffectorIndex(this);
-		if (!bSkipIdxCheck && Idx == INDEX_NONE)
-		{
-			continue;
-		}
-
-		if (!InFunction(Cloner, Idx))
-		{
-			return;
-		}
-	}
-}
-
-void ACEEffectorActor::OnEffectorSubsystemInitialized(const UWorld* World)
-{
-	if (GetWorld() == World)
-	{
-		RegisterToChannel();
-	}
-}
-
-void ACEEffectorActor::RegisterToChannel()
-{
-	UCEEffectorSubsystem::OnSubsystemInitializedDelegate.RemoveAll(this);
-
-	// Register this effector to the effector channel
-	if (UCEEffectorSubsystem* EffectorSubsystem = UCEEffectorSubsystem::Get(GetWorld()))
-	{
-		EffectorSubsystem->RegisterChannelEffector(this);
-	}
-	else
-	{
-		UCEEffectorSubsystem::OnSubsystemInitializedDelegate.AddUObject(this, &ACEEffectorActor::OnEffectorSubsystemInitialized);
-	}
-}
-
 void ACEEffectorActor::OnEffectorTransformed(USceneComponent* InUpdatedComponent, EUpdateTransformFlags InUpdateTransformFlags, ETeleportType InTeleport)
 {
 	OnTransformChanged();
@@ -823,30 +737,6 @@ void ACEEffectorActor::OnEffectorDisabled()
 	OnMagnitudeChanged();
 }
 
-void ACEEffectorActor::OnClonersChanged()
-{
-	TSet<TWeakObjectPtr<ACEClonerActor>> AddedCloners = Cloners.Difference(InternalCloners);
-	TSet<TWeakObjectPtr<ACEClonerActor>> RemovedCloners = InternalCloners.Difference(Cloners);
-
-	// Unlink removed cloners
-	for (const TWeakObjectPtr<ACEClonerActor>& RemovedClonerWeak : RemovedCloners)
-	{
-		if (ACEClonerActor* RemovedCloner = RemovedClonerWeak.Get())
-		{
-			RemovedCloner->UnregisterEffector(this);
-		}
-	}
-
-	// Link added cloners
-	for (const TWeakObjectPtr<ACEClonerActor>& AddedClonerWeak : AddedCloners)
-	{
-		if (ACEClonerActor* AddedCloner = AddedClonerWeak.Get())
-		{
-			AddedCloner->RegisterEffector(this);
-		}
-	}
-}
-
 void ACEEffectorActor::OnModeChanged()
 {
 	ChannelData.Mode = Mode;
@@ -887,41 +777,6 @@ void ACEEffectorActor::OnTypeChanged()
 	OnPlaneChanged();
 }
 
-void ACEEffectorActor::RefreshClonerParameters(ACEClonerActor* InCloner, bool bInForce) const
-{
-	if (InCloner && (bEnabled || bInForce))
-	{
-		InCloner->RequestClonerUpdate();
-	}
-}
-
-void ACEEffectorActor::RefreshClonersParameters(bool bInForce)
-{
-	ForEachCloner([this, bInForce](ACEClonerActor* InCloner, int32 InIndex)
-	{
-		RefreshClonerParameters(InCloner, bInForce);
-		return true;
-	});
-}
-
-void ACEEffectorActor::RegisterToCloners()
-{
-	for (TSet<TWeakObjectPtr<ACEClonerActor>>::TIterator It = InternalCloners.CreateIterator(); It; ++It)
-	{
-		ACEClonerActor* Cloner = It->Get();
-
-		if (!Cloner)
-		{
-			It.RemoveCurrent();
-			continue;
-		}
-
-		Cloner->RegisterEffector(this);
-	}
-
-	OnEffectorChanged();
-}
-
 void ACEEffectorActor::OnEffectorChanged()
 {
 	if (!bEnabled)
@@ -942,32 +797,6 @@ void ACEEffectorActor::OnEffectorChanged()
 	}
 }
 
-void ACEEffectorActor::OnClonerLinked(ACEClonerActor* InCloner, int32 InEffectorIdx)
-{
-	if (!InCloner)
-	{
-		return;
-	}
-
-	InternalCloners.Add(InCloner);
-
-	OnEffectorIdentifierChanged();
-
-	UE_LOG(LogCEEffectorActor, Log, TEXT("%s : Effector linked to Cloner %s"), *GetActorNameOrLabel(), *InCloner->GetActorNameOrLabel());
-}
-
-void ACEEffectorActor::OnClonerUnlinked(ACEClonerActor* InCloner, int32 InEffectorIdx)
-{
-	if (!InCloner)
-	{
-		return;
-	}
-
-	InternalCloners.Remove(InCloner);
-
-	UE_LOG(LogCEEffectorActor, Log, TEXT("%s : Effector unlinked from cloner %s"), *GetActorNameOrLabel(), *InCloner->GetActorNameOrLabel());
-}
-
 void ACEEffectorActor::OnEasingChanged()
 {
 	ChannelData.Easing = Easing;
@@ -1249,7 +1078,11 @@ void ACEEffectorActor::OnForceOptionsChanged()
 
 void ACEEffectorActor::OnForceEnabledChanged()
 {
-	// Refresh cloners to reset clones transform
-	RefreshClonersParameters();
+	if (bEnabled)
+	{
+		// Refresh cloners to reset clones transform
+		OnEffectorRefreshClonerDelegate.Broadcast(this);
+	}
+
 	OnForceOptionsChanged();
 }
diff --git a/Engine/Plugins/Experimental/ClonerEffector/Source/ClonerEffector/Private/Subsystems/CEEffectorSubsystem.cpp b/Engine/Plugins/Experimental/ClonerEffector/Source/ClonerEffector/Private/Subsystems/CEEffectorSubsystem.cpp
--- a/Engine/Plugins/Experimental/ClonerEffector/Source/ClonerEffector/Private/Subsystems/CEEffectorSubsystem.cpp
+++ b/Engine/Plugins/Experimental/ClonerEffector/Source/ClonerEffector/Private/Subsystems/CEEffectorSubsystem.cpp
@@ -10,6 +10,7 @@
 DEFINE_LOG_CATEGORY_STATIC(LogACEEffectorSubsystem, Log, All);
 
 UCEEffectorSubsystem::FOnSubsystemInitialized UCEEffectorSubsystem::OnSubsystemInitializedDelegate;
+UCEEffectorSubsystem::FOnEffectorIdentifierChanged UCEEffectorSubsystem::OnEffectorIdentifierChangedDelegate;
 
 UCEEffectorSubsystem* UCEEffectorSubsystem::Get(const UWorld* InWorld)
 {
@@ -43,20 +44,26 @@ void UCEEffectorSubsystem::PostInitialize()
 
 bool UCEEffectorSubsystem::RegisterChannelEffector(ACEEffectorActor* InEffector)
 {
-	if (!InEffector)
+	if (!IsValid(InEffector))
 	{
 		return false;
 	}
 
-	const int32 EffectorIndex = EffectorsWeak.AddUnique(InEffector);
-	const bool bRegistered = EffectorIndex == EffectorsWeak.Num() - 1;
+	int32 EffectorIndex = INDEX_NONE;
 
-	if (bRegistered)
+	if (!EffectorsWeak.Contains(InEffector))
 	{
+		EffectorIndex = EffectorsWeak.Add(InEffector);
 		UE_LOG(LogACEEffectorSubsystem, Log, TEXT("%s effector registered in channel %i"), *InEffector->GetActorNameOrLabel(), EffectorIndex);
 	}
 
-	return bRegistered;
+	if (InEffector->ChannelData.Identifier != EffectorIndex)
+	{
+		InEffector->ChannelData.Identifier = EffectorIndex;
+		OnEffectorIdentifierChangedDelegate.Broadcast(InEffector);
+	}
+
+	return true;
 }
 
 bool UCEEffectorSubsystem::UnregisterChannelEffector(ACEEffectorActor* InEffector)
@@ -70,6 +77,9 @@ bool UCEEffectorSubsystem::UnregisterChannelEffector(ACEEffectorActor* InEffecto
 
 	if (bUnregistered)
 	{
+		InEffector->ChannelData.Identifier = INDEX_NONE;
+		OnEffectorIdentifierChangedDelegate.Broadcast(InEffector);
+
 		UE_LOG(LogACEEffectorSubsystem, Log, TEXT("%s effector unregistered from channel"), *InEffector->GetActorNameOrLabel());
 	}
 
@@ -106,13 +116,13 @@ void UCEEffectorSubsystem::UpdateEffectorChannel()
 	{
 		ACEEffectorActor* Effector = It->Get();
 
-		if (!Effector)
+		if (!IsValid(Effector))
 		{
 			It.RemoveCurrent();
 			continue;
 		}
 
-		FCEClonerEffectorChannelData& ChannelData = Effector->GetEffectorChannelData();
+		FCEClonerEffectorChannelData& ChannelData = Effector->GetChannelData();
 
 		const bool bIdentifierChanged = ChannelData.Identifier != EffectorIndex;
 
@@ -125,7 +135,7 @@ void UCEEffectorSubsystem::UpdateEffectorChannel()
 		// When changed, update cloners DI linked to this effector
 		if (bIdentifierChanged)
 		{
-			Effector->OnEffectorIdentifierChanged();
+			OnEffectorIdentifierChangedDelegate.Broadcast(Effector);
 		}
 	}
 }
diff --git a/Engine/Plugins/Experimental/ClonerEffector/Source/ClonerEffector/Private/Subsystems/CEEffectorSubsystem.h b/Engine/Plugins/Experimental/ClonerEffector/Source/ClonerEffector/Private/Subsystems/CEEffectorSubsystem.h
--- a/Engine/Plugins/Experimental/ClonerEffector/Source/ClonerEffector/Private/Subsystems/CEEffectorSubsystem.h
+++ b/Engine/Plugins/Experimental/ClonerEffector/Source/ClonerEffector/Private/Subsystems/CEEffectorSubsystem.h
@@ -19,6 +19,10 @@ public:
 	DECLARE_MULTICAST_DELEGATE_OneParam(FOnSubsystemInitialized, const UWorld*)
 	static FOnSubsystemInitialized OnSubsystemInitializedDelegate;
 
+	/** Broadcast when this effector identifier changed to update linked cloners */
+	DECLARE_MULTICAST_DELEGATE_OneParam(FOnEffectorIdentifierChanged, ACEEffectorActor* /** InEffectorActor */)
+	static FOnEffectorIdentifierChanged OnEffectorIdentifierChangedDelegate;
+
 	static inline constexpr TCHAR DataChannelAssetPath[] = TEXT("/Script/Niagara.NiagaraDataChannelAsset'/ClonerEffector/Channels/NDC_Effector.NDC_Effector'");
 
 	/** Get this subsystem instance */
diff --git a/Engine/Plugins/Experimental/ClonerEffector/Source/ClonerEffector/Public/CEClonerEffectorShared.h b/Engine/Plugins/Experimental/ClonerEffector/Source/ClonerEffector/Public/CEClonerEffectorShared.h
--- a/Engine/Plugins/Experimental/ClonerEffector/Source/ClonerEffector/Public/CEClonerEffectorShared.h
+++ b/Engine/Plugins/Experimental/ClonerEffector/Source/ClonerEffector/Public/CEClonerEffectorShared.h
@@ -396,6 +396,7 @@ struct FCEClonerEffectorDataInterfaces
 	explicit FCEClonerEffectorDataInterfaces(const UNiagaraSystem* InSystem);
 	FCEClonerEffectorDataInterfaces() = default;
 
+	void Clear() const;
 	void CopyTo(FCEClonerEffectorDataInterfaces& InOther) const;
 	void Resize(int32 InSize) const;
 	void Remove(int32 InIndex) const;
diff --git a/Engine/Plugins/Experimental/ClonerEffector/Source/ClonerEffector/Public/Cloner/CEClonerActor.h b/Engine/Plugins/Experimental/ClonerEffector/Source/ClonerEffector/Public/Cloner/CEClonerActor.h
--- a/Engine/Plugins/Experimental/ClonerEffector/Source/ClonerEffector/Public/Cloner/CEClonerActor.h
+++ b/Engine/Plugins/Experimental/ClonerEffector/Source/ClonerEffector/Public/Cloner/CEClonerActor.h
@@ -379,10 +379,14 @@ public:
 	void SpawnDefaultActorAttached();
 #endif
 
-	bool GetClonerInitialized() const
-	{
-		return bClonerInitialized;
-	}
+	/** Links a new effector to apply transformation on clones */
+	CLONEREFFECTOR_API bool LinkEffector(ACEEffectorActor* InEffector);
+
+	/** Unlinks the effector and reset the cloner simulation */
+	CLONEREFFECTOR_API bool UnlinkEffector(ACEEffectorActor* InEffector);
+
+	/** Checks if an effector is linked with this cloner */
+	CLONEREFFECTOR_API bool IsEffectorLinked(const ACEEffectorActor* InEffector) const;
 
 protected:
 	//~ Begin UObject
@@ -420,31 +424,21 @@ protected:
 	/** Will force a system update to refresh user parameters */
 	void RequestClonerUpdate(bool bInImmediate = false);
 
-	void UpdateLayoutOptions();
-	void UpdateClonerEffectors();
+	void UpdateLayoutOptions();;
 
 	/** Used by effector actors to apply transformations to this cloner instances */
 	const FCEClonerEffectorDataInterfaces* GetEffectorDataInterfaces() const;
 
-	/** Register a new effector and return the index of this effector to match niagara data interface */
-	int32 RegisterEffector(ACEEffectorActor* InEffector);
-
-	/** Unregister the effector and reorganize the niagara data interface to match with the new array indexes */
-	bool UnregisterEffector(ACEEffectorActor* InEffector);
-
-	/** Checks if an effector is registered with this cloner */
-	bool IsEffectorRegistered(const ACEEffectorActor* InEffector) const;
-
-	/** Gets the index of an effector in this cloner */
-	int32 GetEffectorIndex(ACEEffectorActor* InEffector) const;
-
 	/** For each valid effector that has this cloner registered */
 	void ForEachEffector(TFunctionRef<bool(ACEEffectorActor*, int32)> InFunction);
 
-	void OnClonerTransformed(USceneComponent*, EUpdateTransformFlags, ETeleportType);
-	void OnClonerMeshUpdated();
+	void OnClonerMeshUpdated(UCEClonerComponent* InClonerComponent);
 	void OnClonerSystemChanged();
 
+	void OnEffectorIdentifierChanged(ACEEffectorActor* InEffector);
+	void OnEffectorRefreshCloner(ACEEffectorActor* InEffector);
+
+	void OnEffectorsChanged();
 	void OnEnabledChanged();
 	void OnMeshRenderModeChanged();
 	void OnMeshRendererOptionsChanged();
@@ -637,10 +631,6 @@ private:
 	/** Initiate and perform operation */
 	void InitializeCloner();
 
-	/** Effectors that applies to this cloner, index match niagara data interface arrays */
-	UPROPERTY(Transient)
-	TArray<TWeakObjectPtr<ACEEffectorActor>> Effectors;
-
 	/** Cloner niagara component */
 	UPROPERTY()
 	TObjectPtr<UCEClonerComponent> ClonerComponent;
@@ -649,6 +639,10 @@ private:
 	UPROPERTY()
 	TMap<FName, TObjectPtr<UCEClonerLayoutBase>> LayoutInstances;
 
+	/** Effectors linked to this cloner */
+	UPROPERTY(EditInstanceOnly, Category="Cloner", meta=(DisplayName="Effectors"))
+	TArray<TWeakObjectPtr<ACEEffectorActor>> EffectorsWeak;
+
 	/**
 	* Below properties are deprecated and no longer in use,
 	* they will be migrated to new layout system on load,
@@ -695,6 +689,8 @@ private:
 	bool bClonerInitialized = false;
 
 #if WITH_EDITOR
+	bool bSpawnDefaultActorAttached = false;
+
 	/** Used for PECP */
 	static const TCEPropertyChangeDispatcher<ACEClonerActor> PropertyChangeDispatcher;
 #endif
diff --git a/Engine/Plugins/Experimental/ClonerEffector/Source/ClonerEffector/Public/Cloner/CEClonerComponent.h b/Engine/Plugins/Experimental/ClonerEffector/Source/ClonerEffector/Public/Cloner/CEClonerComponent.h
--- a/Engine/Plugins/Experimental/ClonerEffector/Source/ClonerEffector/Public/Cloner/CEClonerComponent.h
+++ b/Engine/Plugins/Experimental/ClonerEffector/Source/ClonerEffector/Public/Cloner/CEClonerComponent.h
@@ -17,11 +17,11 @@ class UCEClonerComponent : public UNiagaraComponent
 
 public:
 	/** Called when meshes have been updated */
-	DECLARE_MULTICAST_DELEGATE(FOnClonerMeshUpdated)
+	DECLARE_MULTICAST_DELEGATE_OneParam(FOnClonerMeshUpdated, UCEClonerComponent* /** ClonerComponent */)
 	static FOnClonerMeshUpdated OnClonerMeshUpdated;
 
 	/** Called when cloner layout system is loaded */
-	DECLARE_MULTICAST_DELEGATE_OneParam(FOnClonerSystemLoaded, UCEClonerLayoutBase* /** InLayout */)
+	DECLARE_MULTICAST_DELEGATE_TwoParams(FOnClonerSystemLoaded, UCEClonerComponent* /** ClonerComponent */, UCEClonerLayoutBase* /** InLayout */)
 	static FOnClonerSystemLoaded OnClonerSystemLoaded;
 
 	UCEClonerComponent();
diff --git a/Engine/Plugins/Experimental/ClonerEffector/Source/ClonerEffector/Public/Effector/CEEffectorActor.h b/Engine/Plugins/Experimental/ClonerEffector/Source/ClonerEffector/Public/Effector/CEEffectorActor.h
--- a/Engine/Plugins/Experimental/ClonerEffector/Source/ClonerEffector/Public/Effector/CEEffectorActor.h
+++ b/Engine/Plugins/Experimental/ClonerEffector/Source/ClonerEffector/Public/Effector/CEEffectorActor.h
@@ -27,9 +27,6 @@ public:
 
 	ACEEffectorActor();
 
-	CLONEREFFECTOR_API void LinkCloner(ACEClonerActor* InCloner);
-	CLONEREFFECTOR_API void UnlinkCloner(ACEClonerActor* InCloner);
-
 	UFUNCTION(BlueprintCallable, Category="Effector")
 	CLONEREFFECTOR_API void SetType(ECEClonerEffectorType InType);
 
@@ -357,6 +354,10 @@ public:
 #endif
 
 protected:
+	/** Used to trigger a refresh on linked cloner */
+	DECLARE_MULTICAST_DELEGATE_OneParam(FOnEffectorIdentifierChanged, ACEEffectorActor* /** InEffectorActor */)
+	static FOnEffectorIdentifierChanged OnEffectorRefreshClonerDelegate;
+
 	//~ Begin AActor
 	virtual void Destroyed() override;
 	virtual void PostActorCreated() override;
@@ -375,35 +376,12 @@ protected:
 	virtual void PostDuplicate(EDuplicateMode::Type InDuplicateMode) override;
 	//~ End UObject
 
-	/** For each valid cloner that has this effector registered, loops in the internal cloner array */
-	void ForEachCloner(TFunctionRef<bool(ACEClonerActor*, int32)> InFunction, bool bSkipIdxCheck = false);
-
-	void OnEffectorSubsystemInitialized(const UWorld* World);
+	int32 GetChannelIdentifier() const;
 
-	/** Registers this effector to data channel */
-	void RegisterToChannel();
-
-	/** Called when this effector identifier changed to update linked cloners */
-	void OnEffectorIdentifierChanged();
-
-	FCEClonerEffectorChannelData& GetEffectorChannelData()
-	{
-		return ChannelData;
-	}
+	FCEClonerEffectorChannelData& GetChannelData();
 
 	void OnEffectorTransformed(USceneComponent* InUpdatedComponent, EUpdateTransformFlags InUpdateTransformFlags, ETeleportType InTeleport);
 
-	void OnClonerUpdated(ACEClonerActor* InCloner, int32 InEffectorIdx) {}
-	void OnClonerLinked(ACEClonerActor* InCloner, int32 InEffectorIdx);
-	void OnClonerUnlinked(ACEClonerActor* InCloner, int32 InEffectorIdx);
-
-	/** Will update cloner system only if effector is enabled or force update is true, update will happen next tick */
-	void RefreshClonerParameters(ACEClonerActor* InCloner, bool bInForce = false) const;
-	void RefreshClonersParameters(bool bInForce = false);
-
-	/** Will register this effector with linked clones */
-	void RegisterToCloners();
-
 	/** Will update all values of this effector on all clones */
 	void OnEffectorChanged();
 
@@ -413,9 +391,6 @@ protected:
 	/** Called when we disable this effector */
 	void OnEffectorDisabled();
 
-	/** Called when cloners set is updated */
-	void OnClonersChanged();
-
 	/** Update mode for this effector */
 	void OnModeChanged();
 
@@ -472,10 +447,6 @@ protected:
 	UPROPERTY(EditAnywhere, BlueprintReadWrite, Setter="SetMagnitude", Getter="GetMagnitude", Category="Effector", meta=(ClampMin="0", ClampMax="1"))
 	float Magnitude = 1.f;
 
-	/** Cloners affected by this effector */
-	UPROPERTY(EditAnywhere, Category="Effector", meta=(AllowPrivateAccess = "true"))
-	TSet<TWeakObjectPtr<ACEClonerActor>> Cloners;
-
 	/** Type of effector to apply on cloners instances */
 	UPROPERTY(EditAnywhere, BlueprintReadWrite, Setter="SetType", Getter="GetType", Category="Type")
 	ECEClonerEffectorType Type = ECEClonerEffectorType::Sphere;
@@ -637,7 +608,7 @@ private:
 	UPROPERTY()
 	TObjectPtr<UBoxComponent> OuterPlaneComponent;
 
-	/** Internal cloners array */
+	/** Internal cloners array, deprecated since it will be moved to the cloners and emptied out on post load */
 	UPROPERTY(NonTransactional)
 	TSet<TWeakObjectPtr<ACEClonerActor>> InternalCloners;
 
