MotionDesign : ClonerEffector
- Fix warning message shown due to material applied to default cube is missing niagara usage flag (replacing it by custom material with flag set)

#changelist validated
#virtualized

diff --git a/Engine/Plugins/Experimental/ClonerEffector/Source/ClonerEffector/Private/Cloner/CEClonerActor.cpp b/Engine/Plugins/Experimental/ClonerEffector/Source/ClonerEffector/Private/Cloner/CEClonerActor.cpp
--- a/Engine/Plugins/Experimental/ClonerEffector/Source/ClonerEffector/Private/Cloner/CEClonerActor.cpp
+++ b/Engine/Plugins/Experimental/ClonerEffector/Source/ClonerEffector/Private/Cloner/CEClonerActor.cpp
@@ -13,6 +13,7 @@
 #include "Cloner/Layouts/CEClonerSplineLayout.h"
 #include "Effector/CEEffectorActor.h"
 #include "Engine/StaticMeshActor.h"
+#include "Materials/MaterialInterface.h"
 #include "NiagaraDataInterfaceCurve.h"
 #include "NiagaraDataInterfaceSkeletalMesh.h"
 #include "NiagaraMeshRendererProperties.h"
@@ -1288,7 +1289,11 @@ void ACEClonerActor::SpawnDefaultActorAttached()
 	constexpr const TCHAR* DefaultStaticMeshPath = TEXT("/Script/Engine.StaticMesh'/Engine/BasicShapes/Cube.Cube'");
 	UStaticMesh* DefaultStaticMesh = LoadObject<UStaticMesh>(nullptr, DefaultStaticMeshPath);
 
-	if (!DefaultStaticMesh)
+	// Find or load default material
+	constexpr const TCHAR* DefaultMaterialPath = TEXT("/Script/Engine.Material'/ClonerEffector/Materials/DefaultClonerMaterial.DefaultClonerMaterial'");
+	UMaterialInterface* DefaultMaterial = LoadObject<UMaterialInterface>(nullptr, DefaultMaterialPath);
+
+	if (!DefaultStaticMesh || !DefaultMaterial)
 	{
 		return;
 	}
@@ -1303,6 +1308,7 @@ void ACEClonerActor::SpawnDefaultActorAttached()
 	{
 		UStaticMeshComponent* StaticMeshComponent = DefaultActorAttached->GetStaticMeshComponent();
 		StaticMeshComponent->SetStaticMesh(DefaultStaticMesh);
+		StaticMeshComponent->SetMaterial(0, DefaultMaterial);
 
 		DefaultActorAttached->SetMobility(EComponentMobility::Movable);
 		DefaultActorAttached->AttachToActor(this, FAttachmentTransformRules::KeepWorldTransform);
