MotionDesign : ClonerEffector
- Fix crash with cloner component due to use of NewObject while garbage collection is ongoing

#changelist validated
#virtualized

diff --git a/Engine/Plugins/Experimental/ClonerEffector/Source/ClonerEffector/Private/Cloner/CEClonerComponent.cpp b/Engine/Plugins/Experimental/ClonerEffector/Source/ClonerEffector/Private/Cloner/CEClonerComponent.cpp
--- a/Engine/Plugins/Experimental/ClonerEffector/Source/ClonerEffector/Private/Cloner/CEClonerComponent.cpp
+++ b/Engine/Plugins/Experimental/ClonerEffector/Source/ClonerEffector/Private/Cloner/CEClonerComponent.cpp
@@ -84,6 +84,7 @@ void UCEClonerComponent::UpdateClonerRenderState()
 	 * Tree is up to date
 	 */
 	if (!GetAsset()
+		|| IsGarbageCollectingAndLockingUObjectHashTables()
 		|| bClonerMeshesUpdating
 		|| !bClonerMeshesDirty
 		|| ClonerTree.Status != ECEClonerAttachmentStatus::Updated)
@@ -661,56 +662,83 @@ void UCEClonerComponent::UpdateDirtyMeshesAsync()
 
 	// Update baked dynamic meshes on other thread
 	TWeakObjectPtr<UCEClonerComponent> ThisWeak(this);
-	Async(EAsyncExecution::TaskGraph, [ThisWeak, UpdateItems]()->bool
+	Async(EAsyncExecution::TaskGraph, [ThisWeak, UpdateItems]()
 	{
-		if (!ThisWeak.IsValid())
+		UCEClonerComponent* This = ThisWeak.Get();
+
+		if (!This)
 		{
-			return false;
+			return;
 		}
 
 		// update actor baked dynamic meshes
-		for (const FCEClonerAttachmentItem* Item : UpdateItems)
+		bool bSuccess = true;
+		for (FCEClonerAttachmentItem* Item : UpdateItems)
 		{
 			if (!Item || !Item->ItemActor.IsValid())
 			{
 				continue;
 			}
 
-			ThisWeak->UpdateActorBakedDynamicMesh(Item->ItemActor.Get());
+			if (IsGarbageCollectingAndLockingUObjectHashTables())
+			{
+				bSuccess = false;
+				This->DirtyItemAttachments.Add(Item);
+				continue;
+			}
+
+			This->UpdateActorBakedDynamicMesh(Item->ItemActor.Get());
 		}
 
 		// Create baked static mesh on main thread (required)
-		Async(EAsyncExecution::TaskGraphMainThread, [ThisWeak]()
+		Async(EAsyncExecution::TaskGraphMainThread, [ThisWeak, &bSuccess]()
 		{
-			if (!ThisWeak.IsValid())
+			UCEClonerComponent* This = ThisWeak.Get();
+
+			if (!This)
+			{
+				return;
+			}
+
+			if (!bSuccess)
 			{
+				This->OnDirtyMeshesUpdated(false);
 				return;
 			}
 
 			// Update actors baked static mesh
-			for (int32 Idx = 0; Idx < ThisWeak->ClonerTree.RootActors.Num(); Idx++)
+			for (int32 Idx = 0; Idx < This->ClonerTree.RootActors.Num(); Idx++)
 			{
-				AActor* RootActor = ThisWeak->ClonerTree.RootActors[Idx].Get();
-				if (!ThisWeak->ClonerTree.MergedBakedMeshes[Idx].Get())
+				if (IsGarbageCollectingAndLockingUObjectHashTables())
 				{
-					ThisWeak->UpdateRootActorBakedStaticMesh(RootActor);
+					bSuccess = false;
+					break;
+				}
+
+				const UStaticMesh* RootStaticMesh = This->ClonerTree.MergedBakedMeshes[Idx].Get();
+
+				if (!RootStaticMesh)
+				{
+					AActor* RootActor = This->ClonerTree.RootActors[Idx].Get();
+					This->UpdateRootActorBakedStaticMesh(RootActor);
 				}
 			}
 
 			// update niagara asset
-			ThisWeak->OnDirtyMeshesUpdated();
+			This->OnDirtyMeshesUpdated(bSuccess);
 		});
-
-		return true;
 	});
 }
 
-void UCEClonerComponent::OnDirtyMeshesUpdated()
+void UCEClonerComponent::OnDirtyMeshesUpdated(bool bInSuccess)
 {
 	bClonerMeshesUpdating = false;
 
 	// Update niagara parameters
-	UpdateClonerMeshes();
+	if (bInSuccess)
+	{
+		UpdateClonerMeshes();
+	}
 }
 
 void UCEClonerComponent::UpdateActorBakedDynamicMesh(AActor* InActor)
diff --git a/Engine/Plugins/Experimental/ClonerEffector/Source/ClonerEffector/Public/Cloner/CEClonerComponent.h b/Engine/Plugins/Experimental/ClonerEffector/Source/ClonerEffector/Public/Cloner/CEClonerComponent.h
--- a/Engine/Plugins/Experimental/ClonerEffector/Source/ClonerEffector/Public/Cloner/CEClonerComponent.h
+++ b/Engine/Plugins/Experimental/ClonerEffector/Source/ClonerEffector/Public/Cloner/CEClonerComponent.h
@@ -85,7 +85,7 @@ protected:
 
 	/** Runs async update to rebuild dirty meshes */
 	void UpdateDirtyMeshesAsync();
-	void OnDirtyMeshesUpdated();
+	void OnDirtyMeshesUpdated(bool bInSuccess);
 
 	/** Merges all primitive components from an actor to dynamic mesh, does not recurse */
 	void UpdateActorBakedDynamicMesh(AActor* InActor);
