MotionDesign : ClonerEffector
- Changing effector vizualizer size to match standard size

#changelist validated
#virtualized

diff --git a/Engine/Plugins/Experimental/ClonerEffector/Source/ClonerEffector/Public/Effector/CEEffectorActor.h b/Engine/Plugins/Experimental/ClonerEffector/Source/ClonerEffector/Public/Effector/CEEffectorActor.h
--- a/Engine/Plugins/Experimental/ClonerEffector/Source/ClonerEffector/Public/Effector/CEEffectorActor.h
+++ b/Engine/Plugins/Experimental/ClonerEffector/Source/ClonerEffector/Public/Effector/CEEffectorActor.h
@@ -570,7 +570,7 @@ protected:
 
 	/** Thickness of components visualizers */
 	UPROPERTY(EditAnywhere, BlueprintReadWrite, AdvancedDisplay, Setter="SetVisualizerThickness", Getter="GetVisualizerThickness", Category="Effector", meta=(ClampMin="0.1", ClampMax="10.0", Delta="0.1"))
-	float VisualizerThickness = 2.f;
+	float VisualizerThickness = 0.5f;
 
 #if WITH_EDITORONLY_DATA
 	/** Toggle the sprite to visualize and click on this effector */
