MotionDesign : ClonerEffector
- Fix delete effector issue when linked to a cloner
- Fix undo cloner behavior to restore state properly

diff --git a/Engine/Plugins/Experimental/ClonerEffector/Source/ClonerEffector/Private/Cloner/CEClonerActor.cpp b/Engine/Plugins/Experimental/ClonerEffector/Source/ClonerEffector/Private/Cloner/CEClonerActor.cpp
--- a/Engine/Plugins/Experimental/ClonerEffector/Source/ClonerEffector/Private/Cloner/CEClonerActor.cpp
+++ b/Engine/Plugins/Experimental/ClonerEffector/Source/ClonerEffector/Private/Cloner/CEClonerActor.cpp
@@ -81,6 +81,13 @@ FString ACEClonerActor::GetDefaultActorLabel() const
 	return DefaultLabel;
 }
 
+void ACEClonerActor::PostEditUndo()
+{
+	Super::PostEditUndo();
+
+	ForceUpdateCloner();
+}
+
 const TCEPropertyChangeDispatcher<ACEClonerActor> ACEClonerActor::PropertyChangeDispatcher =
 {
 	{ GET_MEMBER_NAME_CHECKED(ACEClonerActor, bEnabled), &ACEClonerActor::OnEnabledChanged },
@@ -1046,10 +1053,20 @@ void ACEClonerActor::OnClonerSystemChanged()
 	UpdateLayoutOptions();
 }
 
-void ACEClonerActor::OnEffectorIdentifierChanged(ACEEffectorActor* InEffector)
+void ACEClonerActor::OnEffectorIdentifierChanged(ACEEffectorActor* InEffector, int32 InOldIdentifier, int32 InNewIdentifier)
 {
 	if (EffectorsWeak.Contains(InEffector))
 	{
+		// Remove effector if it is unregistered from channel
+		if (InNewIdentifier == INDEX_NONE)
+		{
+#if WITH_EDITOR
+			Modify();
+#endif
+
+			EffectorsWeak.Remove(InEffector);
+		}
+
 		OnEffectorsChanged();
 	}
 }
@@ -1091,8 +1108,8 @@ void ACEClonerActor::OnEffectorsChanged()
 		}
 	}
 
-	EffectorDataInterfaces->Clear();
 	TArray<int32>& EffectorIndexArray = EffectorDataInterfaces->GetIndexArray()->GetArrayReference();
+	EffectorIndexArray.Empty(EffectorIndexes.Num());
 
 	for (const int32 EffectorIndex : EffectorIndexes)
 	{
@@ -1223,9 +1240,11 @@ void ACEClonerActor::ForceUpdateCloner()
 {
 	if (ClonerComponent)
 	{
-		constexpr bool bReset = true;
-		ClonerComponent->UpdateClonerAttachmentTree(bReset);
-		TreeUpdateDeltaTime = TreeUpdateInterval;
+		ClonerComponent->UpdateClonerAttachmentTree();
+		ClonerComponent->UpdateClonerRenderState();
+
+		OnLayoutNameChanged();
+		OnEffectorsChanged();
 	}
 }
 
diff --git a/Engine/Plugins/Experimental/ClonerEffector/Source/ClonerEffector/Private/Subsystems/CEEffectorSubsystem.cpp b/Engine/Plugins/Experimental/ClonerEffector/Source/ClonerEffector/Private/Subsystems/CEEffectorSubsystem.cpp
--- a/Engine/Plugins/Experimental/ClonerEffector/Source/ClonerEffector/Private/Subsystems/CEEffectorSubsystem.cpp
+++ b/Engine/Plugins/Experimental/ClonerEffector/Source/ClonerEffector/Private/Subsystems/CEEffectorSubsystem.cpp
@@ -59,8 +59,9 @@ bool UCEEffectorSubsystem::RegisterChannelEffector(ACEEffectorActor* InEffector)
 
 	if (InEffector->ChannelData.Identifier != EffectorIndex)
 	{
+		const int32 OldIdentifier = InEffector->ChannelData.Identifier;
 		InEffector->ChannelData.Identifier = EffectorIndex;
-		OnEffectorIdentifierChangedDelegate.Broadcast(InEffector);
+		OnEffectorIdentifierChangedDelegate.Broadcast(InEffector, OldIdentifier, EffectorIndex);
 	}
 
 	return true;
@@ -77,10 +78,11 @@ bool UCEEffectorSubsystem::UnregisterChannelEffector(ACEEffectorActor* InEffecto
 
 	if (bUnregistered)
 	{
-		InEffector->ChannelData.Identifier = INDEX_NONE;
-		OnEffectorIdentifierChangedDelegate.Broadcast(InEffector);
-
 		UE_LOG(LogACEEffectorSubsystem, Log, TEXT("%s effector unregistered from channel"), *InEffector->GetActorNameOrLabel());
+
+		const int32 OldIdentifier = InEffector->ChannelData.Identifier;
+		InEffector->ChannelData.Identifier = INDEX_NONE;
+		OnEffectorIdentifierChangedDelegate.Broadcast(InEffector, OldIdentifier, InEffector->ChannelData.Identifier);
 	}
 
 	return bUnregistered;
@@ -124,7 +126,8 @@ void UCEEffectorSubsystem::UpdateEffectorChannel()
 
 		FCEClonerEffectorChannelData& ChannelData = Effector->GetChannelData();
 
-		const bool bIdentifierChanged = ChannelData.Identifier != EffectorIndex;
+		const int32 PreviousIdentifier = ChannelData.Identifier;
+		const bool bIdentifierChanged = PreviousIdentifier != EffectorIndex;
 
 		// Set channel before writing
 		ChannelData.Identifier = EffectorIndex++;
@@ -135,7 +138,7 @@ void UCEEffectorSubsystem::UpdateEffectorChannel()
 		// When changed, update cloners DI linked to this effector
 		if (bIdentifierChanged)
 		{
-			OnEffectorIdentifierChangedDelegate.Broadcast(Effector);
+			OnEffectorIdentifierChangedDelegate.Broadcast(Effector, PreviousIdentifier, ChannelData.Identifier);
 		}
 	}
 }
diff --git a/Engine/Plugins/Experimental/ClonerEffector/Source/ClonerEffector/Private/Subsystems/CEEffectorSubsystem.h b/Engine/Plugins/Experimental/ClonerEffector/Source/ClonerEffector/Private/Subsystems/CEEffectorSubsystem.h
--- a/Engine/Plugins/Experimental/ClonerEffector/Source/ClonerEffector/Private/Subsystems/CEEffectorSubsystem.h
+++ b/Engine/Plugins/Experimental/ClonerEffector/Source/ClonerEffector/Private/Subsystems/CEEffectorSubsystem.h
@@ -20,7 +20,7 @@ public:
 	static FOnSubsystemInitialized OnSubsystemInitializedDelegate;
 
 	/** Broadcast when this effector identifier changed to update linked cloners */
-	DECLARE_MULTICAST_DELEGATE_OneParam(FOnEffectorIdentifierChanged, ACEEffectorActor* /** InEffectorActor */)
+	DECLARE_MULTICAST_DELEGATE_ThreeParams(FOnEffectorIdentifierChanged, ACEEffectorActor* /** InEffectorActor */, int32 /** OldIdentifier */, int32 /** NewIdentifier */)
 	static FOnEffectorIdentifierChanged OnEffectorIdentifierChangedDelegate;
 
 	static inline constexpr TCHAR DataChannelAssetPath[] = TEXT("/Script/Niagara.NiagaraDataChannelAsset'/ClonerEffector/Channels/NDC_Effector.NDC_Effector'");
diff --git a/Engine/Plugins/Experimental/ClonerEffector/Source/ClonerEffector/Public/Cloner/CEClonerActor.h b/Engine/Plugins/Experimental/ClonerEffector/Source/ClonerEffector/Public/Cloner/CEClonerActor.h
--- a/Engine/Plugins/Experimental/ClonerEffector/Source/ClonerEffector/Public/Cloner/CEClonerActor.h
+++ b/Engine/Plugins/Experimental/ClonerEffector/Source/ClonerEffector/Public/Cloner/CEClonerActor.h
@@ -393,6 +393,7 @@ protected:
 	virtual void PostInitProperties() override;
 	virtual void PostLoad() override;
 #if WITH_EDITOR
+	virtual void PostEditUndo() override;
 	virtual void PostEditChangeProperty(FPropertyChangedEvent& InPropertyChangedEvent) override;
 #endif
 	//~ End UObject
@@ -435,7 +436,7 @@ protected:
 	void OnClonerMeshUpdated(UCEClonerComponent* InClonerComponent);
 	void OnClonerSystemChanged();
 
-	void OnEffectorIdentifierChanged(ACEEffectorActor* InEffector);
+	void OnEffectorIdentifierChanged(ACEEffectorActor* InEffector, int32 InOldIdentifier, int32 InNewIdentifier);
 	void OnEffectorRefreshCloner(ACEEffectorActor* InEffector);
 
 	void OnEffectorsChanged();
