MotionDesign : Cloner
- Disabling use of bounds when focusing cloner (pressing F) due to GPU sim the bounds cannot be dynamically computed, large fixed bounds are used to avoid particles culling

diff --git a/Engine/Plugins/Experimental/ClonerEffector/Source/ClonerEffector/Private/Cloner/CEClonerComponent.cpp b/Engine/Plugins/Experimental/ClonerEffector/Source/ClonerEffector/Private/Cloner/CEClonerComponent.cpp
--- a/Engine/Plugins/Experimental/ClonerEffector/Source/ClonerEffector/Private/Cloner/CEClonerComponent.cpp
+++ b/Engine/Plugins/Experimental/ClonerEffector/Source/ClonerEffector/Private/Cloner/CEClonerComponent.cpp
@@ -49,6 +49,9 @@ UCEClonerComponent::UCEClonerComponent()
 #if WITH_EDITOR
 	// Do not show bounding box around cloner for better visibility
 	SetIsVisualizationComponent(true);
+
+	// Disable use of bounds to focus to avoid de-zoom
+	SetIgnoreBoundsForEditorFocus(true);
 #endif
 
 	bIsEditorOnly = false;
