MotionDesign : ClonerEffector
- Add particles color support, the cloner has a default color and effectors also have a color, on overlap those color are blend together, default material uses color for preview/debug
- Add twist support for grid and honeycomb layout, you can select the axis and twist strength

#virtualized

diff --git a/Engine/Plugins/Experimental/ClonerEffector/Source/ClonerEffector/Private/CEClonerEffectorShared.cpp b/Engine/Plugins/Experimental/ClonerEffector/Source/ClonerEffector/Private/CEClonerEffectorShared.cpp
--- a/Engine/Plugins/Experimental/ClonerEffector/Source/ClonerEffector/Private/CEClonerEffectorShared.cpp
+++ b/Engine/Plugins/Experimental/ClonerEffector/Source/ClonerEffector/Private/CEClonerEffectorShared.cpp
@@ -77,6 +77,7 @@ void FCEClonerEffectorChannelData::Write(UNiagaraDataChannelWriter* InWriter) co
 	InWriter->WriteVector(ScaleName, Identifier, Scale);
 	InWriter->WriteFloat(FrequencyName, Identifier, Frequency);
 	InWriter->WriteVector(PanName, Identifier, Pan);
+	InWriter->WriteLinearColor(ColorName, Identifier, Color);
 
 	/** Forces */
 	InWriter->WriteFloat(OrientationForceRateName, Identifier, OrientationForceRate);
@@ -114,6 +115,7 @@ void FCEClonerEffectorChannelData::Read(const UNiagaraDataChannelReader* InReade
 	Scale = InReader->ReadVector(ScaleName, Identifier, bIsValid);
 	Frequency = InReader->ReadFloat(FrequencyName, Identifier, bIsValid);
 	Pan = InReader->ReadVector(PanName, Identifier, bIsValid);
+	Color = InReader->ReadLinearColor(ColorName, Identifier, bIsValid);
 
 	/** Forces */
 	OrientationForceRate = InReader->ReadFloat(OrientationForceRateName, Identifier, bIsValid);
diff --git a/Engine/Plugins/Experimental/ClonerEffector/Source/ClonerEffector/Private/Cloner/CEClonerActor.cpp b/Engine/Plugins/Experimental/ClonerEffector/Source/ClonerEffector/Private/Cloner/CEClonerActor.cpp
--- a/Engine/Plugins/Experimental/ClonerEffector/Source/ClonerEffector/Private/Cloner/CEClonerActor.cpp
+++ b/Engine/Plugins/Experimental/ClonerEffector/Source/ClonerEffector/Private/Cloner/CEClonerActor.cpp
@@ -92,6 +92,7 @@ const TCEPropertyChangeDispatcher<ACEClonerActor> ACEClonerActor::PropertyChange
 {
 	{ GET_MEMBER_NAME_CHECKED(ACEClonerActor, bEnabled), &ACEClonerActor::OnEnabledChanged },
 	{ GET_MEMBER_NAME_CHECKED(ACEClonerActor, Seed), &ACEClonerActor::OnSeedChanged },
+	{ GET_MEMBER_NAME_CHECKED(ACEClonerActor, Color), &ACEClonerActor::OnColorChanged },
 	{ GET_MEMBER_NAME_CHECKED(ACEClonerActor, EffectorsWeak), &ACEClonerActor::OnEffectorsChanged },
 	/** Layout */
 	{ GET_MEMBER_NAME_CHECKED(ACEClonerActor, LayoutName), &ACEClonerActor::OnLayoutNameChanged },
@@ -349,6 +350,7 @@ void ACEClonerActor::UpdateLayoutOptions()
 {
 	OnMeshRenderModeChanged();
 	OnSeedChanged();
+	OnColorChanged();
 
 	// update layouts options
 	if (UCEClonerLayoutBase* CurrentLayout = ClonerComponent->GetClonerActiveLayout())
@@ -445,6 +447,17 @@ void ACEClonerActor::SetSeed(int32 InSeed)
 	OnSeedChanged();
 }
 
+void ACEClonerActor::SetColor(const FLinearColor& InColor)
+{
+	if (InColor == Color)
+	{
+		return;
+	}
+
+	Color = InColor;
+	OnColorChanged();
+}
+
 void ACEClonerActor::SetDeltaStepEnabled(bool bInEnabled)
 {
 	if (bDeltaStepEnabled == bInEnabled)
@@ -696,6 +709,16 @@ void ACEClonerActor::OnSeedChanged()
 	RequestClonerUpdate();
 }
 
+void ACEClonerActor::OnColorChanged()
+{
+	if (!ClonerComponent || !bEnabled)
+	{
+		return;
+	}
+
+	ClonerComponent->SetColorParameter(TEXT("EffectorDefaultColor"), Color);
+}
+
 void ACEClonerActor::OnProgressChanged()
 {
 	if (!ClonerComponent || !bEnabled)
diff --git a/Engine/Plugins/Experimental/ClonerEffector/Source/ClonerEffector/Private/Cloner/Layouts/CEClonerGridLayout.cpp b/Engine/Plugins/Experimental/ClonerEffector/Source/ClonerEffector/Private/Cloner/Layouts/CEClonerGridLayout.cpp
--- a/Engine/Plugins/Experimental/ClonerEffector/Source/ClonerEffector/Private/Cloner/Layouts/CEClonerGridLayout.cpp
+++ b/Engine/Plugins/Experimental/ClonerEffector/Source/ClonerEffector/Private/Cloner/Layouts/CEClonerGridLayout.cpp
@@ -112,6 +112,28 @@ void UCEClonerGridLayout::SetTextureConstraint(const FCEClonerGridConstraintText
 	UpdateLayoutParameters();
 }
 
+void UCEClonerGridLayout::SetTwistFactor(float InFactor)
+{
+	if (TwistFactor == InFactor)
+	{
+		return;
+	}
+
+	TwistFactor = InFactor;
+	UpdateLayoutParameters();
+}
+
+void UCEClonerGridLayout::SetTwistAxis(ENiagaraOrientationAxis InAxis)
+{
+	if (TwistAxis == InAxis)
+	{
+		return;
+	}
+
+	TwistAxis = InAxis;
+	UpdateLayoutParameters();
+}
+
 #if WITH_EDITOR
 const TCEPropertyChangeDispatcher<UCEClonerGridLayout> UCEClonerGridLayout::PropertyChangeDispatcher =
 {
@@ -126,6 +148,8 @@ const TCEPropertyChangeDispatcher<UCEClonerGridLayout> UCEClonerGridLayout::Prop
 	{ GET_MEMBER_NAME_CHECKED(UCEClonerGridLayout, SphereConstraint), &UCEClonerGridLayout::OnLayoutPropertyChanged },
 	{ GET_MEMBER_NAME_CHECKED(UCEClonerGridLayout, CylinderConstraint), &UCEClonerGridLayout::OnLayoutPropertyChanged },
 	{ GET_MEMBER_NAME_CHECKED(UCEClonerGridLayout, TextureConstraint), &UCEClonerGridLayout::OnLayoutPropertyChanged },
+	{ GET_MEMBER_NAME_CHECKED(UCEClonerGridLayout, TwistFactor), &UCEClonerGridLayout::OnLayoutPropertyChanged },
+	{ GET_MEMBER_NAME_CHECKED(UCEClonerGridLayout, TwistAxis), &UCEClonerGridLayout::OnLayoutPropertyChanged },
 };
 
 void UCEClonerGridLayout::PostEditChangeProperty(FPropertyChangedEvent& InPropertyChangedEvent)
@@ -178,4 +202,9 @@ void UCEClonerGridLayout::OnLayoutParametersChanged(UCEClonerComponent* InCompon
 	static const FNiagaraVariable ConstraintTextureSamplerVar(FNiagaraTypeDefinition(UNiagaraDataInterfaceTexture::StaticClass()), TEXT("ConstraintTextureSampler"));
 	UNiagaraDataInterfaceTexture* TextureSamplerDI = Cast<UNiagaraDataInterfaceTexture>(ExposedParameters.GetDataInterface(ConstraintTextureSamplerVar));
 	TextureSamplerDI->SetTexture(TextureConstraint.Texture.Get());
+
+	InComponent->SetFloatParameter(TEXT("TwistFactor"), TwistFactor);
+
+	static const FNiagaraVariable TwistAxisVar(FNiagaraTypeDefinition(StaticEnum<ENiagaraOrientationAxis>()), TEXT("TwistAxis"));
+	ExposedParameters.SetParameterValue<int32>(static_cast<int32>(TwistAxis), TwistAxisVar);
 }
diff --git a/Engine/Plugins/Experimental/ClonerEffector/Source/ClonerEffector/Private/Cloner/Layouts/CEClonerHoneycombLayout.cpp b/Engine/Plugins/Experimental/ClonerEffector/Source/ClonerEffector/Private/Cloner/Layouts/CEClonerHoneycombLayout.cpp
--- a/Engine/Plugins/Experimental/ClonerEffector/Source/ClonerEffector/Private/Cloner/Layouts/CEClonerHoneycombLayout.cpp
+++ b/Engine/Plugins/Experimental/ClonerEffector/Source/ClonerEffector/Private/Cloner/Layouts/CEClonerHoneycombLayout.cpp
@@ -7,13 +7,18 @@
 
 void UCEClonerHoneycombLayout::SetPlane(ECEClonerPlane InPlane)
 {
+	if (Plane == ECEClonerPlane::Custom)
+	{
+		return;
+	}
+
 	if (Plane == InPlane)
 	{
 		return;
 	}
 
 	Plane = InPlane;
-	UpdateLayoutParameters();
+	OnTwistAxisChanged();
 }
 
 void UCEClonerHoneycombLayout::SetWidthCount(int32 InWidthCount)
@@ -82,16 +87,40 @@ void UCEClonerHoneycombLayout::SetWidthSpacing(float InWidthSpacing)
 	UpdateLayoutParameters();
 }
 
+void UCEClonerHoneycombLayout::SetTwistFactor(float InFactor)
+{
+	if (TwistFactor == InFactor)
+	{
+		return;
+	}
+
+	TwistFactor = InFactor;
+	UpdateLayoutParameters();
+}
+
+void UCEClonerHoneycombLayout::SetTwistAxis(ENiagaraOrientationAxis InAxis)
+{
+	if (TwistAxis == InAxis)
+	{
+		return;
+	}
+
+	TwistAxis = InAxis;
+	OnTwistAxisChanged();
+}
+
 #if WITH_EDITOR
 const TCEPropertyChangeDispatcher<UCEClonerHoneycombLayout> UCEClonerHoneycombLayout::PropertyChangeDispatcher =
 {
-	{ GET_MEMBER_NAME_CHECKED(UCEClonerHoneycombLayout, Plane), &UCEClonerHoneycombLayout::OnLayoutPropertyChanged },
+	{ GET_MEMBER_NAME_CHECKED(UCEClonerHoneycombLayout, Plane), &UCEClonerHoneycombLayout::OnTwistAxisChanged },
 	{ GET_MEMBER_NAME_CHECKED(UCEClonerHoneycombLayout, WidthCount), &UCEClonerHoneycombLayout::OnLayoutPropertyChanged },
 	{ GET_MEMBER_NAME_CHECKED(UCEClonerHoneycombLayout, HeightCount), &UCEClonerHoneycombLayout::OnLayoutPropertyChanged },
 	{ GET_MEMBER_NAME_CHECKED(UCEClonerHoneycombLayout, WidthOffset), &UCEClonerHoneycombLayout::OnLayoutPropertyChanged },
 	{ GET_MEMBER_NAME_CHECKED(UCEClonerHoneycombLayout, HeightOffset), &UCEClonerHoneycombLayout::OnLayoutPropertyChanged },
 	{ GET_MEMBER_NAME_CHECKED(UCEClonerHoneycombLayout, WidthSpacing), &UCEClonerHoneycombLayout::OnLayoutPropertyChanged },
 	{ GET_MEMBER_NAME_CHECKED(UCEClonerHoneycombLayout, HeightSpacing), &UCEClonerHoneycombLayout::OnLayoutPropertyChanged },
+	{ GET_MEMBER_NAME_CHECKED(UCEClonerHoneycombLayout, TwistFactor), &UCEClonerHoneycombLayout::OnLayoutPropertyChanged },
+	{ GET_MEMBER_NAME_CHECKED(UCEClonerHoneycombLayout, TwistAxis), &UCEClonerHoneycombLayout::OnTwistAxisChanged },
 };
 
 void UCEClonerHoneycombLayout::PostEditChangeProperty(FPropertyChangedEvent& InPropertyChangedEvent)
@@ -102,6 +131,34 @@ void UCEClonerHoneycombLayout::PostEditChangeProperty(FPropertyChangedEvent& InP
 }
 #endif
 
+void UCEClonerHoneycombLayout::OnTwistAxisChanged()
+{
+	// Restrict twist axis to plane axis
+	if (Plane == ECEClonerPlane::XY)
+	{
+		if (TwistAxis == ENiagaraOrientationAxis::ZAxis)
+		{
+			TwistAxis = ENiagaraOrientationAxis::XAxis;
+		}
+	}
+	else if (Plane == ECEClonerPlane::XZ)
+	{
+		if (TwistAxis == ENiagaraOrientationAxis::YAxis)
+		{
+			TwistAxis = ENiagaraOrientationAxis::XAxis;
+		}
+	}
+	else if (Plane == ECEClonerPlane::YZ)
+	{
+		if (TwistAxis == ENiagaraOrientationAxis::XAxis)
+		{
+			TwistAxis = ENiagaraOrientationAxis::YAxis;
+		}
+	}
+
+	UpdateLayoutParameters();
+}
+
 void UCEClonerHoneycombLayout::OnLayoutParametersChanged(UCEClonerComponent* InComponent)
 {
 	Super::OnLayoutParametersChanged(InComponent);
@@ -121,4 +178,9 @@ void UCEClonerHoneycombLayout::OnLayoutParametersChanged(UCEClonerComponent* InC
 	FNiagaraUserRedirectionParameterStore& ExposedParameters = InComponent->GetAsset()->GetExposedParameters();
 	static const FNiagaraVariable HoneycombPlaneVar(FNiagaraTypeDefinition(StaticEnum<ECEClonerPlane>()), TEXT("HoneycombPlane"));
 	ExposedParameters.SetParameterValue<int32>(static_cast<int32>(Plane), HoneycombPlaneVar);
+
+	InComponent->SetFloatParameter(TEXT("TwistFactor"), TwistFactor);
+
+	static const FNiagaraVariable TwistAxisVar(FNiagaraTypeDefinition(StaticEnum<ENiagaraOrientationAxis>()), TEXT("TwistAxis"));
+	ExposedParameters.SetParameterValue<int32>(static_cast<int32>(TwistAxis), TwistAxisVar);
 }
diff --git a/Engine/Plugins/Experimental/ClonerEffector/Source/ClonerEffector/Private/Effector/CEEffectorActor.cpp b/Engine/Plugins/Experimental/ClonerEffector/Source/ClonerEffector/Private/Effector/CEEffectorActor.cpp
--- a/Engine/Plugins/Experimental/ClonerEffector/Source/ClonerEffector/Private/Effector/CEEffectorActor.cpp
+++ b/Engine/Plugins/Experimental/ClonerEffector/Source/ClonerEffector/Private/Effector/CEEffectorActor.cpp
@@ -118,6 +118,7 @@ TCEPropertyChangeDispatcher<ACEEffectorActor> ACEEffectorActor::PropertyChangeDi
 	{ GET_MEMBER_NAME_CHECKED(ACEEffectorActor, Magnitude), &ACEEffectorActor::OnMagnitudeChanged },
 	{ GET_MEMBER_NAME_CHECKED(ACEEffectorActor, VisualizerThickness), &ACEEffectorActor::OnVisualizerThicknessChanged },
 	{ GET_MEMBER_NAME_CHECKED(ACEEffectorActor, bVisualizerSpriteVisible), &ACEEffectorActor::OnVisualizerSpriteVisibleChanged },
+	{ GET_MEMBER_NAME_CHECKED(ACEEffectorActor, Color), &ACEEffectorActor::OnColorChanged },
 	/** Type */
 	{ GET_MEMBER_NAME_CHECKED(ACEEffectorActor, Type), &ACEEffectorActor::OnTypeChanged },
 	{ GET_MEMBER_NAME_CHECKED(ACEEffectorActor, Easing), &ACEEffectorActor::OnEasingChanged },
@@ -700,6 +701,17 @@ void ACEEffectorActor::SetVisualizerSpriteVisible(bool bInVisible)
 }
 #endif
 
+void ACEEffectorActor::SetColor(const FLinearColor& InColor)
+{
+	if (Color == InColor)
+	{
+		return;
+	}
+
+	Color = InColor;
+	OnColorChanged();
+}
+
 void ACEEffectorActor::OnEffectorTransformed(USceneComponent* InUpdatedComponent, EUpdateTransformFlags InUpdateTransformFlags, ETeleportType InTeleport)
 {
 	OnTransformChanged();
@@ -791,6 +803,7 @@ void ACEEffectorActor::OnEffectorChanged()
 		OnTransformChanged();
 		OnMagnitudeChanged();
 		OnForceOptionsChanged();
+		OnColorChanged();
 		// Editor
 		OnVisualizerThicknessChanged();
 		OnVisualizerSpriteVisibleChanged();
@@ -1018,6 +1031,11 @@ void ACEEffectorActor::OnVisualizerSpriteVisibleChanged()
 #endif
 }
 
+void ACEEffectorActor::OnColorChanged()
+{
+	ChannelData.Color = Color;
+}
+
 void ACEEffectorActor::OnForceOptionsChanged()
 {
 	if (bOrientationForceEnabled)
diff --git a/Engine/Plugins/Experimental/ClonerEffector/Source/ClonerEffector/Public/CEClonerEffectorShared.h b/Engine/Plugins/Experimental/ClonerEffector/Source/ClonerEffector/Public/CEClonerEffectorShared.h
--- a/Engine/Plugins/Experimental/ClonerEffector/Source/ClonerEffector/Public/CEClonerEffectorShared.h
+++ b/Engine/Plugins/Experimental/ClonerEffector/Source/ClonerEffector/Public/CEClonerEffectorShared.h
@@ -330,6 +330,7 @@ struct FCEClonerEffectorChannelData
 	static constexpr const TCHAR* ScaleName = TEXT("Scale");
 	static constexpr const TCHAR* FrequencyName = TEXT("Frequency");
 	static constexpr const TCHAR* PanName = TEXT("Pan");
+	static constexpr const TCHAR* ColorName = TEXT("Color");
 
 	/** Forces */
 	static constexpr const TCHAR* OrientationForceRateName = TEXT("OrientationForceRate");
@@ -363,6 +364,7 @@ struct FCEClonerEffectorChannelData
 	FVector Scale = FVector::OneVector;
 	float Frequency = 1.f;
 	FVector Pan = FVector::ZeroVector;
+	FLinearColor Color = FLinearColor::Transparent;
 
 	/** Forces parameters */
 	float OrientationForceRate = 0.f;
diff --git a/Engine/Plugins/Experimental/ClonerEffector/Source/ClonerEffector/Public/Cloner/CEClonerActor.h b/Engine/Plugins/Experimental/ClonerEffector/Source/ClonerEffector/Public/Cloner/CEClonerActor.h
--- a/Engine/Plugins/Experimental/ClonerEffector/Source/ClonerEffector/Public/Cloner/CEClonerActor.h
+++ b/Engine/Plugins/Experimental/ClonerEffector/Source/ClonerEffector/Public/Cloner/CEClonerActor.h
@@ -91,6 +91,15 @@ public:
 		return Seed;
 	}
 
+	UFUNCTION(BlueprintCallable, Category="Cloner")
+	CLONEREFFECTOR_API void SetColor(const FLinearColor& InColor);
+
+	UFUNCTION(BlueprintPure, Category="Cloner")
+	const FLinearColor& GetColor() const
+	{
+		return Color;
+	}
+
 	UFUNCTION(BlueprintCallable, Category="Cloner")
 	CLONEREFFECTOR_API void SetDeltaStepEnabled(bool bInEnabled);
 
@@ -425,7 +434,7 @@ protected:
 	/** Will force a system update to refresh user parameters */
 	void RequestClonerUpdate(bool bInImmediate = false);
 
-	void UpdateLayoutOptions();;
+	void UpdateLayoutOptions();
 
 	/** Used by effector actors to apply transformations to this cloner instances */
 	const FCEClonerEffectorDataInterfaces* GetEffectorDataInterfaces() const;
@@ -445,6 +454,7 @@ protected:
 	void OnMeshRendererOptionsChanged();
 	void OnDefaultMeshesChanged();
 	void OnSeedChanged();
+	void OnColorChanged();
 	void OnProgressChanged();
 	void OnDeltaStepChanged();
 	void OnRangeOptionsChanged();
@@ -475,6 +485,10 @@ protected:
 	UPROPERTY(EditAnywhere, BlueprintReadWrite, Setter="SetSeed", Getter="GetSeed", Category="Cloner")
 	int32 Seed = 0;
 
+	/** Cloner color when unaffected by effectors, color will be passed down to the material (ParticleColor) */
+	UPROPERTY(EditAnywhere, BlueprintReadWrite, Setter, Getter, Category="Cloner")
+	FLinearColor Color = FLinearColor::White;
+
 	/** Indicates how we select the mesh to render on each clones */
 	UPROPERTY(EditAnywhere, BlueprintReadWrite, Setter="SetMeshRenderMode", Getter="GetMeshRenderMode", Category="Renderer")
 	ECEClonerMeshRenderMode MeshRenderMode = ECEClonerMeshRenderMode::Iterate;
diff --git a/Engine/Plugins/Experimental/ClonerEffector/Source/ClonerEffector/Public/Cloner/Layouts/CEClonerGridLayout.h b/Engine/Plugins/Experimental/ClonerEffector/Source/ClonerEffector/Public/Cloner/Layouts/CEClonerGridLayout.h
--- a/Engine/Plugins/Experimental/ClonerEffector/Source/ClonerEffector/Public/Cloner/Layouts/CEClonerGridLayout.h
+++ b/Engine/Plugins/Experimental/ClonerEffector/Source/ClonerEffector/Public/Cloner/Layouts/CEClonerGridLayout.h
@@ -121,6 +121,24 @@ public:
 		return TextureConstraint;
 	}
 
+	UFUNCTION(BlueprintCallable, Category="Cloner|Layout|Grid")
+	CLONEREFFECTOR_API void SetTwistFactor(float InFactor);
+
+	UFUNCTION(BlueprintPure, Category="Cloner|Layout|Grid")
+	float GetTwistFactor() const
+	{
+		return TwistFactor;
+	}
+
+	UFUNCTION(BlueprintCallable, Category="Cloner|Layout|Grid")
+	CLONEREFFECTOR_API void SetTwistAxis(ENiagaraOrientationAxis InAxis);
+
+	UFUNCTION(BlueprintPure, Category="Cloner|Layout|Grid")
+	ENiagaraOrientationAxis GetTwistAxis() const
+	{
+		return TwistAxis;
+	}
+
 protected:
 	//~ Begin UObject
 #if WITH_EDITOR
@@ -165,6 +183,12 @@ protected:
 	UPROPERTY(EditInstanceOnly, BlueprintReadWrite, Setter="SetTextureConstraint", Getter="GetTextureConstraint", Category="Layout", meta=(EditCondition="Constraint == ECEClonerGridConstraint::Texture", EditConditionHides))
 	FCEClonerGridConstraintTexture TextureConstraint;
 
+	UPROPERTY(EditInstanceOnly, BlueprintReadWrite, Setter, Getter, Category="Layout")
+	float TwistFactor = 0.f;
+
+	UPROPERTY(EditInstanceOnly, BlueprintReadWrite, Setter, Getter, Category="Layout")
+	ENiagaraOrientationAxis TwistAxis = ENiagaraOrientationAxis::ZAxis;
+
 private:
 #if WITH_EDITOR
 	/** Used for PECP */
diff --git a/Engine/Plugins/Experimental/ClonerEffector/Source/ClonerEffector/Public/Cloner/Layouts/CEClonerHoneycombLayout.h b/Engine/Plugins/Experimental/ClonerEffector/Source/ClonerEffector/Public/Cloner/Layouts/CEClonerHoneycombLayout.h
--- a/Engine/Plugins/Experimental/ClonerEffector/Source/ClonerEffector/Public/Cloner/Layouts/CEClonerHoneycombLayout.h
+++ b/Engine/Plugins/Experimental/ClonerEffector/Source/ClonerEffector/Public/Cloner/Layouts/CEClonerHoneycombLayout.h
@@ -85,6 +85,24 @@ public:
 		return WidthSpacing;
 	}
 
+	UFUNCTION(BlueprintCallable, Category="Cloner|Layout|Honeycomb")
+	CLONEREFFECTOR_API void SetTwistFactor(float InFactor);
+
+	UFUNCTION(BlueprintPure, Category="Cloner|Layout|Honeycomb")
+	float GetTwistFactor() const
+	{
+		return TwistFactor;
+	}
+
+	UFUNCTION(BlueprintCallable, Category="Cloner|Layout|Honeycomb")
+	CLONEREFFECTOR_API void SetTwistAxis(ENiagaraOrientationAxis InAxis);
+
+	UFUNCTION(BlueprintPure, Category="Cloner|Layout|Honeycomb")
+	ENiagaraOrientationAxis GetTwistAxis() const
+	{
+		return TwistAxis;
+	}
+
 protected:
 	//~ Begin UObject
 #if WITH_EDITOR
@@ -92,11 +110,13 @@ protected:
 #endif
 	//~ End UObject
 
+	void OnTwistAxisChanged();
+
 	//~ Begin UCEClonerLayoutBase
 	virtual void OnLayoutParametersChanged(UCEClonerComponent* InComponent) override;
 	//~ End UCEClonerLayoutBase
 
-	UPROPERTY(EditInstanceOnly, BlueprintReadWrite, Setter="SetPlane", Getter="GetPlane", Category="Layout")
+	UPROPERTY(EditInstanceOnly, BlueprintReadWrite, Setter="SetPlane", Getter="GetPlane", Category="Layout", meta=(InvalidEnumValues="Custom"))
 	ECEClonerPlane Plane = ECEClonerPlane::XY;
 
 	UPROPERTY(EditInstanceOnly, BlueprintReadWrite, Setter="SetWidthCount", Getter="GetWidthCount", Category="Layout")
@@ -117,6 +137,12 @@ protected:
 	UPROPERTY(EditInstanceOnly, BlueprintReadWrite, Setter="SetHeightSpacing", Getter="GetHeightSpacing", Category="Layout")
 	float HeightSpacing = 105.f;
 
+	UPROPERTY(EditInstanceOnly, BlueprintReadWrite, Setter, Getter, Category="Layout")
+	float TwistFactor = 0.f;
+
+	UPROPERTY(EditInstanceOnly, BlueprintReadWrite, Setter, Getter, Category="Layout")
+	ENiagaraOrientationAxis TwistAxis = ENiagaraOrientationAxis::XAxis;
+
 private:
 #if WITH_EDITOR
 	/** Used for PECP */
diff --git a/Engine/Plugins/Experimental/ClonerEffector/Source/ClonerEffector/Public/Effector/CEEffectorActor.h b/Engine/Plugins/Experimental/ClonerEffector/Source/ClonerEffector/Public/Effector/CEEffectorActor.h
--- a/Engine/Plugins/Experimental/ClonerEffector/Source/ClonerEffector/Public/Effector/CEEffectorActor.h
+++ b/Engine/Plugins/Experimental/ClonerEffector/Source/ClonerEffector/Public/Effector/CEEffectorActor.h
@@ -135,15 +135,6 @@ public:
 		return bEnabled;
 	}
 
-	UFUNCTION(BlueprintCallable, Category="Effector")
-	CLONEREFFECTOR_API void SetVisualizerThickness(float InThickness);
-
-	UFUNCTION(BlueprintPure, Category="Effector")
-	float GetVisualizerThickness() const
-	{
-		return VisualizerThickness;
-	}
-
 	UFUNCTION(BlueprintCallable, Category="Effector")
 	CLONEREFFECTOR_API void SetMode(ECEClonerEffectorMode InMode);
 
@@ -162,186 +153,204 @@ public:
 		return TargetActorWeak.Get();
 	}
 
-	UFUNCTION(BlueprintCallable, Category="Cloner")
+	UFUNCTION(BlueprintCallable, Category="Effector")
 	CLONEREFFECTOR_API void SetOrientationForceEnabled(bool bInForceEnabled);
 
-	UFUNCTION(BlueprintPure, Category="Cloner")
+	UFUNCTION(BlueprintPure, Category="Effector")
 	bool GetOrientationForceEnabled() const
 	{
 		return bOrientationForceEnabled;
 	}
 
-	UFUNCTION(BlueprintCallable, Category="Cloner")
+	UFUNCTION(BlueprintCallable, Category="Effector")
 	CLONEREFFECTOR_API void SetOrientationForceRate(float InForceOrientationRate);
 
-	UFUNCTION(BlueprintPure, Category="Cloner")
+	UFUNCTION(BlueprintPure, Category="Effector")
 	float GetOrientationForceRate() const
 	{
 		return OrientationForceRate;
 	}
 
-	UFUNCTION(BlueprintCallable, Category="Cloner")
+	UFUNCTION(BlueprintCallable, Category="Effector")
 	CLONEREFFECTOR_API void SetOrientationForceMin(const FVector& InForceOrientationMin);
 
-	UFUNCTION(BlueprintPure, Category="Cloner")
+	UFUNCTION(BlueprintPure, Category="Effector")
 	const FVector& GetOrientationForceMin() const
 	{
 		return OrientationForceMin;
 	}
 
-	UFUNCTION(BlueprintCallable, Category="Cloner")
+	UFUNCTION(BlueprintCallable, Category="Effector")
 	CLONEREFFECTOR_API void SetOrientationForceMax(const FVector& InForceOrientationMax);
 
-	UFUNCTION(BlueprintPure, Category="Cloner")
+	UFUNCTION(BlueprintPure, Category="Effector")
 	const FVector& GetOrientationForceMax() const
 	{
 		return OrientationForceMax;
 	}
 
-	UFUNCTION(BlueprintCallable, Category="Cloner")
+	UFUNCTION(BlueprintCallable, Category="Effector")
 	CLONEREFFECTOR_API void SetVortexForceEnabled(bool bInForceEnabled);
 
-	UFUNCTION(BlueprintPure, Category="Cloner")
+	UFUNCTION(BlueprintPure, Category="Effector")
 	bool GetVortexForceEnabled() const
 	{
 		return bVortexForceEnabled;
 	}
 
-	UFUNCTION(BlueprintCallable, Category="Cloner")
+	UFUNCTION(BlueprintCallable, Category="Effector")
 	CLONEREFFECTOR_API void SetVortexForceAmount(float InForceVortexAmount);
 
-	UFUNCTION(BlueprintPure, Category="Cloner")
+	UFUNCTION(BlueprintPure, Category="Effector")
 	float GetVortexForceAmount() const
 	{
 		return VortexForceAmount;
 	}
 
-	UFUNCTION(BlueprintCallable, Category="Cloner")
+	UFUNCTION(BlueprintCallable, Category="Effector")
 	CLONEREFFECTOR_API void SetVortexForceAxis(const FVector& InForceVortexAxis);
 
-	UFUNCTION(BlueprintPure, Category="Cloner")
+	UFUNCTION(BlueprintPure, Category="Effector")
 	const FVector& GetVortexForceAxis() const
 	{
 		return VortexForceAxis;
 	}
 
-	UFUNCTION(BlueprintCallable, Category="Cloner")
+	UFUNCTION(BlueprintCallable, Category="Effector")
 	CLONEREFFECTOR_API void SetCurlNoiseForceEnabled(bool bInForceEnabled);
 
-	UFUNCTION(BlueprintPure, Category="Cloner")
+	UFUNCTION(BlueprintPure, Category="Effector")
 	bool GetCurlNoiseForceEnabled() const
 	{
 		return bCurlNoiseForceEnabled;
 	}
 
-	UFUNCTION(BlueprintCallable, Category="Cloner")
+	UFUNCTION(BlueprintCallable, Category="Effector")
 	CLONEREFFECTOR_API void SetCurlNoiseForceStrength(float InForceCurlNoiseStrength);
 
-	UFUNCTION(BlueprintPure, Category="Cloner")
+	UFUNCTION(BlueprintPure, Category="Effector")
 	float GetCurlNoiseForceStrength() const
 	{
 		return CurlNoiseForceStrength;
 	}
 
-	UFUNCTION(BlueprintCallable, Category="Cloner")
+	UFUNCTION(BlueprintCallable, Category="Effector")
 	CLONEREFFECTOR_API void SetCurlNoiseForceFrequency(float InForceCurlNoiseFrequency);
 
-	UFUNCTION(BlueprintPure, Category="Cloner")
+	UFUNCTION(BlueprintPure, Category="Effector")
 	float GetCurlNoiseForceFrequency() const
 	{
 		return CurlNoiseForceFrequency;
 	}
 
-	UFUNCTION(BlueprintCallable, Category="Cloner")
+	UFUNCTION(BlueprintCallable, Category="Effector")
 	CLONEREFFECTOR_API void SetAttractionForceEnabled(bool bInForceEnabled);
 
-	UFUNCTION(BlueprintPure, Category="Cloner")
+	UFUNCTION(BlueprintPure, Category="Effector")
 	bool GetAttractionForceEnabled() const
 	{
 		return bAttractionForceEnabled;
 	}
 
-	UFUNCTION(BlueprintCallable, Category="Cloner")
+	UFUNCTION(BlueprintCallable, Category="Effector")
 	CLONEREFFECTOR_API void SetAttractionForceStrength(float InForceStrength);
 
-	UFUNCTION(BlueprintPure, Category="Cloner")
+	UFUNCTION(BlueprintPure, Category="Effector")
 	float GetAttractionForceStrength() const
 	{
 		return AttractionForceStrength;
 	}
 
-	UFUNCTION(BlueprintCallable, Category="Cloner")
+	UFUNCTION(BlueprintCallable, Category="Effector")
 	CLONEREFFECTOR_API void SetAttractionForceFalloff(float InForceFalloff);
 
-	UFUNCTION(BlueprintPure, Category="Cloner")
+	UFUNCTION(BlueprintPure, Category="Effector")
 	float GetAttractionForceFalloff() const
 	{
 		return AttractionForceFalloff;
 	}
 
-	UFUNCTION(BlueprintCallable, Category="Cloner")
+	UFUNCTION(BlueprintCallable, Category="Effector")
 	CLONEREFFECTOR_API void SetGravityForceEnabled(bool bInForceEnabled);
 
-	UFUNCTION(BlueprintPure, Category="Cloner")
+	UFUNCTION(BlueprintPure, Category="Effector")
 	bool GetGravityForceEnabled() const
 	{
 		return bGravityForceEnabled;
 	}
 
-	UFUNCTION(BlueprintCallable, Category="Cloner")
+	UFUNCTION(BlueprintCallable, Category="Effector")
 	CLONEREFFECTOR_API void SetGravityForceAcceleration(const FVector& InAcceleration);
 
-	UFUNCTION(BlueprintPure, Category="Cloner")
+	UFUNCTION(BlueprintPure, Category="Effector")
 	FVector GetGravityForceAcceleration() const
 	{
 		return GravityForceAcceleration;
 	}
 
-	UFUNCTION(BlueprintCallable, Category="Cloner")
+	UFUNCTION(BlueprintCallable, Category="Effector")
 	CLONEREFFECTOR_API void SetLocationStrength(const FVector& InStrength);
 
-	UFUNCTION(BlueprintPure, Category="Cloner")
+	UFUNCTION(BlueprintPure, Category="Effector")
 	FVector GetLocationStrength() const
 	{
 		return LocationStrength;
 	}
 
-	UFUNCTION(BlueprintCallable, Category="Cloner")
+	UFUNCTION(BlueprintCallable, Category="Effector")
 	CLONEREFFECTOR_API void SetRotationStrength(const FRotator& InStrength);
 
-	UFUNCTION(BlueprintPure, Category="Cloner")
+	UFUNCTION(BlueprintPure, Category="Effector")
 	FRotator GetRotationStrength() const
 	{
 		return RotationStrength;
 	}
 
-	UFUNCTION(BlueprintCallable, Category="Cloner")
+	UFUNCTION(BlueprintCallable, Category="Effector")
 	CLONEREFFECTOR_API void SetScaleStrength(const FVector& InStrength);
 
-	UFUNCTION(BlueprintPure, Category="Cloner")
+	UFUNCTION(BlueprintPure, Category="Effector")
 	FVector GetScaleStrength() const
 	{
 		return ScaleStrength;
 	}
 
-	UFUNCTION(BlueprintCallable, Category="Cloner")
+	UFUNCTION(BlueprintCallable, Category="Effector")
 	CLONEREFFECTOR_API void SetPan(const FVector& InPan);
 
-	UFUNCTION(BlueprintPure, Category="Cloner")
+	UFUNCTION(BlueprintPure, Category="Effector")
 	FVector GetPan() const
 	{
 		return Pan;
 	}
 
-	UFUNCTION(BlueprintCallable, Category="Cloner")
+	UFUNCTION(BlueprintCallable, Category="Effector")
 	CLONEREFFECTOR_API void SetFrequency(float InFrequency);
 
-	UFUNCTION(BlueprintPure, Category="Cloner")
+	UFUNCTION(BlueprintPure, Category="Effector")
 	float GetFrequency() const
 	{
 		return Frequency;
 	}
 
+	UFUNCTION(BlueprintCallable, Category="Effector")
+	CLONEREFFECTOR_API void SetColor(const FLinearColor& InColor);
+
+	UFUNCTION(BlueprintPure, Category="Effector")
+	const FLinearColor& GetColor() const
+	{
+		return Color;
+	}
+
+	UFUNCTION(BlueprintCallable, Category="Effector")
+	CLONEREFFECTOR_API void SetVisualizerThickness(float InThickness);
+
+	UFUNCTION(BlueprintPure, Category="Effector")
+	float GetVisualizerThickness() const
+	{
+		return VisualizerThickness;
+	}
+
 #if WITH_EDITOR
 	UFUNCTION(BlueprintCallable, Category="Effector")
 	CLONEREFFECTOR_API void SetVisualizerSpriteVisible(bool bInVisible);
@@ -439,6 +448,9 @@ protected:
 	/** Update sprite visibility of this effector */
 	void OnVisualizerSpriteVisibleChanged();
 
+	/** Update particle color affected by this effector */
+	void OnColorChanged();
+
 	/** Is this effector enabled/disabled on linked cloners */
 	UPROPERTY(EditAnywhere, BlueprintReadWrite, Setter="SetEnabled", Getter="GetEnabled", Category="Effector")
 	bool bEnabled = true;
@@ -447,6 +459,10 @@ protected:
 	UPROPERTY(EditAnywhere, BlueprintReadWrite, Setter="SetMagnitude", Getter="GetMagnitude", Category="Effector", meta=(ClampMin="0", ClampMax="1"))
 	float Magnitude = 1.f;
 
+	/** Affected clones color passed over to material */
+	UPROPERTY(EditAnywhere, BlueprintReadWrite, Setter, Getter, Category="Effector")
+	FLinearColor Color = FLinearColor::White;
+
 	/** Type of effector to apply on cloners instances */
 	UPROPERTY(EditAnywhere, BlueprintReadWrite, Setter="SetType", Getter="GetType", Category="Type")
 	ECEClonerEffectorType Type = ECEClonerEffectorType::Sphere;
