MotionDesign : ClonerEffector
- Add options to override meshes materials by a custom material on cloner in renderer category
- Add options to visualize effectors on cloner in general category

diff --git a/Engine/Plugins/Experimental/ClonerEffector/Source/ClonerEffector/Private/Cloner/CEClonerActor.cpp b/Engine/Plugins/Experimental/ClonerEffector/Source/ClonerEffector/Private/Cloner/CEClonerActor.cpp
--- a/Engine/Plugins/Experimental/ClonerEffector/Source/ClonerEffector/Private/Cloner/CEClonerActor.cpp
+++ b/Engine/Plugins/Experimental/ClonerEffector/Source/ClonerEffector/Private/Cloner/CEClonerActor.cpp
@@ -46,6 +46,10 @@ ACEClonerActor::ACEClonerActor()
 	LifetimeScaleCurve.AddKey(0, 1.f);
 	LifetimeScaleCurve.AddKey(1, 0.f);
 
+	// Default override material
+	static ConstructorHelpers::FObjectFinder<UMaterialInterface> DefaultMaterialFinder(DefaultMaterialPath);
+	OverrideMaterial = DefaultMaterialFinder.Object;
+
 	if (!IsTemplate())
 	{
 #if WITH_EDITOR
@@ -93,6 +97,7 @@ const TCEPropertyChangeDispatcher<ACEClonerActor> ACEClonerActor::PropertyChange
 	{ GET_MEMBER_NAME_CHECKED(ACEClonerActor, bEnabled), &ACEClonerActor::OnEnabledChanged },
 	{ GET_MEMBER_NAME_CHECKED(ACEClonerActor, Seed), &ACEClonerActor::OnSeedChanged },
 	{ GET_MEMBER_NAME_CHECKED(ACEClonerActor, Color), &ACEClonerActor::OnColorChanged },
+	{ GET_MEMBER_NAME_CHECKED(ACEClonerActor, bVisualizeEffectors), &ACEClonerActor::OnOverrideMaterialChanged },
 	{ GET_MEMBER_NAME_CHECKED(ACEClonerActor, EffectorsWeak), &ACEClonerActor::OnEffectorsChanged },
 	/** Layout */
 	{ GET_MEMBER_NAME_CHECKED(ACEClonerActor, LayoutName), &ACEClonerActor::OnLayoutNameChanged },
@@ -115,6 +120,8 @@ const TCEPropertyChangeDispatcher<ACEClonerActor> ACEClonerActor::PropertyChange
 	{ GET_MEMBER_NAME_CHECKED(ACEClonerActor, MeshFacingMode), &ACEClonerActor::OnMeshRendererOptionsChanged },
 	{ GET_MEMBER_NAME_CHECKED(ACEClonerActor, bMeshCastShadows), &ACEClonerActor::OnMeshRendererOptionsChanged },
 	{ GET_MEMBER_NAME_CHECKED(ACEClonerActor, DefaultMeshes), &ACEClonerActor::OnDefaultMeshesChanged },
+	{ GET_MEMBER_NAME_CHECKED(ACEClonerActor, bUseOverrideMaterial), &ACEClonerActor::OnOverrideMaterialChanged },
+	{ GET_MEMBER_NAME_CHECKED(ACEClonerActor, OverrideMaterial), &ACEClonerActor::OnOverrideMaterialChanged },
 	/** Progress */
 	{ GET_MEMBER_NAME_CHECKED(ACEClonerActor, bInvertProgress), &ACEClonerActor::OnProgressChanged },
 	{ GET_MEMBER_NAME_CHECKED(ACEClonerActor, Progress), &ACEClonerActor::OnProgressChanged },
@@ -360,6 +367,8 @@ void ACEClonerActor::UpdateLayoutOptions()
 
 	OnMeshRendererOptionsChanged();
 	OnDefaultMeshesChanged();
+	OnOverrideMaterialChanged();
+
 	OnDeltaStepChanged();
 	OnRangeOptionsChanged();
 	OnSpawnOptionsChanged();
@@ -436,6 +445,28 @@ TArray<UStaticMesh*> ACEClonerActor::BP_GetDefaultMeshes() const
 	return Meshes;
 }
 
+void ACEClonerActor::SetUseOverrideMaterial(bool bInOverride)
+{
+	if (bUseOverrideMaterial == bInOverride)
+	{
+		return;
+	}
+
+	bUseOverrideMaterial = bInOverride;
+	OnOverrideMaterialChanged();
+}
+
+void ACEClonerActor::SetOverrideMaterial(UMaterialInterface* InMaterial)
+{
+	if (OverrideMaterial == InMaterial)
+	{
+		return;
+	}
+
+	OverrideMaterial = InMaterial;
+	OnOverrideMaterialChanged();
+}
+
 void ACEClonerActor::SetSeed(int32 InSeed)
 {
 	if (InSeed == Seed)
@@ -458,6 +489,17 @@ void ACEClonerActor::SetColor(const FLinearColor& InColor)
 	OnColorChanged();
 }
 
+void ACEClonerActor::SetVisualizeEffectors(bool bInVisualize)
+{
+	if (bVisualizeEffectors == bInVisualize)
+	{
+		return;
+	}
+
+	bVisualizeEffectors = bInVisualize;
+	OnOverrideMaterialChanged();
+}
+
 void ACEClonerActor::SetDeltaStepEnabled(bool bInEnabled)
 {
 	if (bDeltaStepEnabled == bInEnabled)
@@ -1235,6 +1277,19 @@ void ACEClonerActor::OnDefaultMeshesChanged()
 	TreeUpdateDeltaTime = TreeUpdateInterval;
 }
 
+void ACEClonerActor::OnOverrideMaterialChanged()
+{
+	if (ClonerComponent)
+	{
+		UMaterialInterface* OverrideMeshesMaterial = bVisualizeEffectors
+			? LoadObject<UMaterialInterface>(nullptr, DefaultMaterialPath)
+			: OverrideMaterial.Get();
+
+		ClonerComponent->SetOverrideMeshesMaterial(OverrideMeshesMaterial);
+		ClonerComponent->SetUseOverrideMeshesMaterial(bUseOverrideMaterial || bVisualizeEffectors);
+	}
+}
+
 #if WITH_EDITOR
 void ACEClonerActor::OnReduceMotionGhostingChanged()
 {
@@ -1332,7 +1387,6 @@ void ACEClonerActor::SpawnDefaultActorAttached()
 	UStaticMesh* DefaultStaticMesh = LoadObject<UStaticMesh>(nullptr, DefaultStaticMeshPath);
 
 	// Find or load default material
-	constexpr const TCHAR* DefaultMaterialPath = TEXT("/Script/Engine.Material'/ClonerEffector/Materials/DefaultClonerMaterial.DefaultClonerMaterial'");
 	UMaterialInterface* DefaultMaterial = LoadObject<UMaterialInterface>(nullptr, DefaultMaterialPath);
 
 	if (!DefaultStaticMesh || !DefaultMaterial)
diff --git a/Engine/Plugins/Experimental/ClonerEffector/Source/ClonerEffector/Private/Cloner/CEClonerComponent.cpp b/Engine/Plugins/Experimental/ClonerEffector/Source/ClonerEffector/Private/Cloner/CEClonerComponent.cpp
--- a/Engine/Plugins/Experimental/ClonerEffector/Source/ClonerEffector/Private/Cloner/CEClonerComponent.cpp
+++ b/Engine/Plugins/Experimental/ClonerEffector/Source/ClonerEffector/Private/Cloner/CEClonerComponent.cpp
@@ -1169,13 +1169,9 @@ void UCEClonerComponent::UpdateRootActorBakedStaticMesh(AActor* InRootActor)
 						UMaterialInterface* MaterialInterface = BakedMaterial.Get();
 						int32 MatIdx = CombineMaterials.Add(MaterialInterface);
 #if WITH_EDITOR
-						if (MaterialInterface)
+						if (MaterialInterface && !IsMaterialUsageFlagSet(MaterialInterface))
 						{
-							const UMaterial* Material = MaterialInterface->GetMaterial();
-							if (Material && !Material->GetUsageByFlag(EMaterialUsage::MATUSAGE_NiagaraMeshParticles))
-							{
-								MaterialsMissingNiagaraUsageFlag.Add(MatIdx);
-							}
+							MaterialsMissingNiagaraUsageFlag.Add(MatIdx);
 						}
 #endif
 					}
@@ -1209,14 +1205,7 @@ void UCEClonerComponent::UpdateRootActorBakedStaticMesh(AActor* InRootActor)
 			}
 		}
 
-		if (ReadOnlyMaterialCount > 0)
-		{
-			FNotificationInfo NotificationInfo(FText::Format(LOCTEXT("MaterialsMissingUsageFlag", "Detected {0} read-only material(s) with missing niagara usage flag required to work properly with cloner (See logs)"), ReadOnlyMaterialCount));
-			NotificationInfo.ExpireDuration = 5.f;
-			NotificationInfo.bFireAndForget = true;
-			NotificationInfo.Image = FAppStyle::GetBrush("Icons.WarningWithColor");
-			FSlateNotificationManager::Get().AddNotification(NotificationInfo);
-		}
+		ShowMaterialWarning(ReadOnlyMaterialCount);
 	}
 #endif
 
@@ -1297,7 +1286,7 @@ bool UCEClonerComponent::IsAllMergedMeshesValid() const
 	return true;
 }
 
-bool UCEClonerComponent::IsMaterialDirtyable(const UMaterialInterface* InMaterial) const
+bool UCEClonerComponent::IsMaterialDirtyable(const UMaterialInterface* InMaterial)
 {
 	const UMaterial* BaseMaterial = InMaterial->GetMaterial();
 	const FString ContentFolder = FPaths::ConvertRelativePathToFull(FPaths::ProjectContentDir());
@@ -1313,15 +1302,127 @@ bool UCEClonerComponent::IsMaterialDirtyable(const UMaterialInterface* InMateria
 	return bTransientPackage || bContentFolder;
 }
 
-void UCEClonerComponent::UpdateClonerMeshes()
+bool UCEClonerComponent::IsMaterialUsageFlagSet(const UMaterialInterface* InMaterial)
 {
-	const AActor* ClonerActor = GetOwner();
-	if (!ClonerActor)
+	if (InMaterial)
+	{
+		if (const UMaterial* Material = InMaterial->GetMaterial())
+		{
+			return Material->GetUsageByFlag(EMaterialUsage::MATUSAGE_NiagaraMeshParticles);
+		}
+	}
+
+	return false;
+}
+
+#if WITH_EDITOR
+void UCEClonerComponent::ShowMaterialWarning(int32 InMaterialCount)
+{
+	if (InMaterialCount > 0)
+	{
+		FNotificationInfo NotificationInfo(FText::Format(LOCTEXT("MaterialsMissingUsageFlag", "Detected {0} read-only material(s) with missing niagara usage flag required to work properly with cloner (See logs)"), InMaterialCount));
+		NotificationInfo.ExpireDuration = 5.f;
+		NotificationInfo.bFireAndForget = true;
+		NotificationInfo.Image = FAppStyle::GetBrush("Icons.WarningWithColor");
+
+		FSlateNotificationManager::Get().AddNotification(NotificationInfo);
+	}
+}
+#endif // WITH_EDITOR
+
+TArray<FNiagaraMeshMaterialOverride> UCEClonerComponent::GetOverrideMeshesMaterials() const
+{
+	TArray<FNiagaraMeshMaterialOverride> MaterialOverrides;
+
+	if (bUseClonerMeshesOverrideMaterial)
+	{
+		int32 MaterialCount = 0;
+
+		for (const TWeakObjectPtr<AActor>& RootActor : ClonerTree.RootActors)
+		{
+			if (!RootActor.IsValid())
+			{
+				continue;
+			}
+
+			const FCEClonerAttachmentItem* AttachmentItem = ClonerTree.ItemAttachmentMap.Find(RootActor);
+
+			if (!AttachmentItem)
+			{
+				continue;
+			}
+
+			MaterialCount += AttachmentItem->BakedMaterials.Num();
+		}
+
+		// Set same material for all available slots
+		MaterialOverrides.Reserve(MaterialCount);
+		UMaterialInterface* OverrideMaterial = ClonerMeshesOverrideMaterial.Get();
+
+		for (int32 Index = 0; Index < MaterialCount; Index++)
+		{
+			FNiagaraMeshMaterialOverride MaterialOverride;
+			MaterialOverride.ExplicitMat = OverrideMaterial;
+			MaterialOverrides.Add(MaterialOverride);
+		}
+	}
+
+	return MaterialOverrides;
+}
+
+void UCEClonerComponent::SetUseOverrideMeshesMaterial(bool bInOverride)
+{
+	if (bUseClonerMeshesOverrideMaterial == bInOverride)
 	{
 		return;
 	}
 
-	if (!bClonerMeshesDirty)
+	bUseClonerMeshesOverrideMaterial = bInOverride;
+
+	if (ClonerMeshesOverrideMaterial.IsValid())
+	{
+		OnOverrideMeshesMaterialChanged();
+	}
+}
+
+void UCEClonerComponent::SetOverrideMeshesMaterial(UMaterialInterface* InMaterial)
+{
+	if (InMaterial && !IsMaterialUsageFlagSet(InMaterial))
+	{
+		UE_LOG(LogCEClonerComponent, Warning, TEXT("%s : The override material (%s) you wish to use does not have the required usage flag (bUsedWithNiagaraMeshParticles) to work with the cloner, enable the flag on the material and save the asset"), *GetOwner()->GetActorNameOrLabel(), *InMaterial->GetMaterial()->GetPathName());
+
+#if WITH_EDITOR
+		ShowMaterialWarning(1);
+#endif
+
+		InMaterial = nullptr;
+	}
+
+	if (ClonerMeshesOverrideMaterial == InMaterial)
+	{
+		return;
+	}
+
+	ClonerMeshesOverrideMaterial = InMaterial;
+
+	if (bUseClonerMeshesOverrideMaterial)
+	{
+		OnOverrideMeshesMaterialChanged();
+	}
+}
+
+void UCEClonerComponent::OnOverrideMeshesMaterialChanged()
+{
+	if (!bClonerMeshesUpdating && !bClonerMeshesDirty)
+	{
+		UpdateClonerMeshes();
+	}
+}
+
+void UCEClonerComponent::UpdateClonerMeshes()
+{
+	const AActor* ClonerActor = GetOwner();
+	if (!ClonerActor)
 	{
 		return;
 	}
@@ -1339,29 +1440,36 @@ void UCEClonerComponent::UpdateClonerMeshes()
 		return;
 	}
 
-	// Set new number of meshes in renderer
-	static const FName MeshNumName(TEXT("MeshNum"));
-	SetIntParameter(MeshNumName, ClonerTree.MergedBakedMeshes.Num());
+	// Set material override
+	MeshRenderer->OverrideMaterials = GetOverrideMeshesMaterials();
+	MeshRenderer->bOverrideMaterials = bUseClonerMeshesOverrideMaterial;
 
-	// Resize mesh array properly
-	if (MeshRenderer->Meshes.Num() > ClonerTree.MergedBakedMeshes.Num())
+	if (bClonerMeshesDirty)
 	{
-		MeshRenderer->Meshes.SetNum(ClonerTree.MergedBakedMeshes.Num());
-	}
+		// Set new number of meshes in renderer
+		static const FName MeshNumName(TEXT("MeshNum"));
+		SetIntParameter(MeshNumName, ClonerTree.MergedBakedMeshes.Num());
 
-	// Set baked meshes in mesh renderer array
-	for (int32 Idx = 0; Idx < ClonerTree.MergedBakedMeshes.Num(); Idx++)
-	{
-		UStaticMesh* StaticMesh = ClonerTree.MergedBakedMeshes[Idx];
-		FNiagaraMeshRendererMeshProperties& MeshProperties = !MeshRenderer->Meshes.IsValidIndex(Idx) ? MeshRenderer->Meshes.AddDefaulted_GetRef() : MeshRenderer->Meshes[Idx];
-		MeshProperties.Mesh = StaticMesh && StaticMesh->GetNumTriangles(0) > 0 ? StaticMesh : nullptr;
+		// Resize mesh array properly
+		if (MeshRenderer->Meshes.Num() > ClonerTree.MergedBakedMeshes.Num())
+		{
+			MeshRenderer->Meshes.SetNum(ClonerTree.MergedBakedMeshes.Num());
+		}
 
-		if (ClonerTree.RootActors.IsValidIndex(Idx))
+		// Set baked meshes in mesh renderer array
+		for (int32 Idx = 0; Idx < ClonerTree.MergedBakedMeshes.Num(); Idx++)
 		{
-			if (const FCEClonerAttachmentItem* Item = ClonerTree.ItemAttachmentMap.Find(ClonerTree.RootActors[Idx]))
+			UStaticMesh* StaticMesh = ClonerTree.MergedBakedMeshes[Idx];
+			FNiagaraMeshRendererMeshProperties& MeshProperties = !MeshRenderer->Meshes.IsValidIndex(Idx) ? MeshRenderer->Meshes.AddDefaulted_GetRef() : MeshRenderer->Meshes[Idx];
+			MeshProperties.Mesh = StaticMesh && StaticMesh->GetNumTriangles(0) > 0 ? StaticMesh : nullptr;
+
+			if (ClonerTree.RootActors.IsValidIndex(Idx))
 			{
-				MeshProperties.Rotation = Item->ActorTransform.Rotator();
-				MeshProperties.Scale = Item->ActorTransform.GetScale3D();
+				if (const FCEClonerAttachmentItem* Item = ClonerTree.ItemAttachmentMap.Find(ClonerTree.RootActors[Idx]))
+				{
+					MeshProperties.Rotation = Item->ActorTransform.Rotator();
+					MeshProperties.Scale = Item->ActorTransform.GetScale3D();
+				}
 			}
 		}
 	}
diff --git a/Engine/Plugins/Experimental/ClonerEffector/Source/ClonerEffector/Public/Cloner/CEClonerActor.h b/Engine/Plugins/Experimental/ClonerEffector/Source/ClonerEffector/Public/Cloner/CEClonerActor.h
--- a/Engine/Plugins/Experimental/ClonerEffector/Source/ClonerEffector/Public/Cloner/CEClonerActor.h
+++ b/Engine/Plugins/Experimental/ClonerEffector/Source/ClonerEffector/Public/Cloner/CEClonerActor.h
@@ -30,6 +30,7 @@ class ACEClonerActor : public AActor
 
 public:
 	static inline const FString DefaultLabel = TEXT("Cloner");
+	static inline constexpr const TCHAR* DefaultMaterialPath = TEXT("/Script/Engine.Material'/ClonerEffector/Materials/DefaultClonerMaterial.DefaultClonerMaterial'");
 
 	ACEClonerActor();
 
@@ -82,6 +83,24 @@ public:
 	UFUNCTION(BlueprintPure, Category="Cloner", meta=(DisplayName="GetDefaultMeshes"))
 	CLONEREFFECTOR_API TArray<UStaticMesh*> BP_GetDefaultMeshes() const;
 
+	UFUNCTION(BlueprintCallable, Category="Cloner")
+	CLONEREFFECTOR_API void SetUseOverrideMaterial(bool bInOverride);
+
+	UFUNCTION(BlueprintPure, Category="Cloner")
+	bool GetUseOverrideMaterial() const
+	{
+		return bUseOverrideMaterial;
+	}
+
+	UFUNCTION(BlueprintCallable, Category="Cloner")
+	CLONEREFFECTOR_API void SetOverrideMaterial(UMaterialInterface* InMaterial);
+
+	UFUNCTION(BlueprintPure, Category="Cloner")
+	UMaterialInterface* GetOverrideMaterial() const
+	{
+		return OverrideMaterial;
+	}
+
 	UFUNCTION(BlueprintCallable, Category="Cloner")
 	CLONEREFFECTOR_API void SetSeed(int32 InSeed);
 
@@ -100,6 +119,15 @@ public:
 		return Color;
 	}
 
+	UFUNCTION(BlueprintCallable, Category="Cloner")
+	CLONEREFFECTOR_API void SetVisualizeEffectors(bool bInVisualize);
+
+	UFUNCTION(BlueprintPure, Category="Cloner")
+	bool GetVisualizeEffectors() const
+	{
+		return bVisualizeEffectors;
+	}
+
 	UFUNCTION(BlueprintCallable, Category="Cloner")
 	CLONEREFFECTOR_API void SetDeltaStepEnabled(bool bInEnabled);
 
@@ -453,6 +481,7 @@ protected:
 	void OnMeshRenderModeChanged();
 	void OnMeshRendererOptionsChanged();
 	void OnDefaultMeshesChanged();
+	void OnOverrideMaterialChanged();
 	void OnSeedChanged();
 	void OnColorChanged();
 	void OnProgressChanged();
@@ -489,6 +518,10 @@ protected:
 	UPROPERTY(EditAnywhere, BlueprintReadWrite, Setter, Getter, Category="Cloner")
 	FLinearColor Color = FLinearColor::White;
 
+	/** Switches materials to show effectors applied on this cloner based on their color property */
+	UPROPERTY(EditAnywhere, BlueprintReadWrite, Setter="SetVisualizeEffectors", Getter="GetVisualizeEffectors", Category="Cloner")
+	bool bVisualizeEffectors;
+
 	/** Indicates how we select the mesh to render on each clones */
 	UPROPERTY(EditAnywhere, BlueprintReadWrite, Setter="SetMeshRenderMode", Getter="GetMeshRenderMode", Category="Renderer")
 	ECEClonerMeshRenderMode MeshRenderMode = ECEClonerMeshRenderMode::Iterate;
@@ -505,6 +538,14 @@ protected:
 	UPROPERTY(EditAnywhere, BlueprintReadWrite, Setter="SetDefaultMeshes", Getter="GetDefaultMeshes", BlueprintSetter="BP_SetDefaultMeshes", BlueprintGetter="BP_GetDefaultMeshes", Category="Renderer")
 	TArray<TObjectPtr<UStaticMesh>> DefaultMeshes;
 
+	/** Whether to override meshes materials with another material */
+	UPROPERTY(EditAnywhere, BlueprintReadWrite, Setter="SetUseOverrideMaterial", Getter="GetUseOverrideMaterial", Category="Renderer")
+	bool bUseOverrideMaterial = false;
+
+	/** The override materials that will be set instead of meshes materials, bVisualizeEffectors must be disabled */
+	UPROPERTY(EditAnywhere, BlueprintReadWrite, Setter, Getter, Category="Renderer", meta=(EditCondition="bUseOverrideMaterial", EditConditionHides))
+	TObjectPtr<UMaterialInterface> OverrideMaterial;
+
 	/** How many times do we spawn clones */
 	UPROPERTY(EditAnywhere, BlueprintReadWrite, Setter="SetSpawnLoopMode", Getter="GetSpawnLoopMode", DisplayName="LoopMode", Category="Spawn")
 	ECEClonerSpawnLoopMode SpawnLoopMode = ECEClonerSpawnLoopMode::Once;
@@ -609,11 +650,11 @@ protected:
 
 #if WITH_EDITORONLY_DATA
 	/** Toggle the sprite to visualize and click on this cloner */
-	UPROPERTY(EditAnywhere, BlueprintReadWrite, AdvancedDisplay, Category="Cloner", meta=(AllowPrivateAccess = "true"))
+	UPROPERTY(EditAnywhere, BlueprintReadWrite, AdvancedDisplay, Category="Cloner", meta=(AllowPrivateAccess="true"))
 	bool bVisualizerSpriteVisible = true;
 
 	/** Reduces the r.TSR.ShadingRejection.Flickering.Period from 3 (default) to 1 if enabled to avoid ghosting artifacts when moving */
-	UPROPERTY(EditAnywhere, BlueprintReadWrite, AdvancedDisplay, Category="Cloner", meta=(AllowPrivateAccess = "true"))
+	UPROPERTY(EditAnywhere, BlueprintReadWrite, AdvancedDisplay, Category="Cloner", meta=(AllowPrivateAccess="true"))
 	bool bReduceMotionGhosting = false;
 #endif
 
diff --git a/Engine/Plugins/Experimental/ClonerEffector/Source/ClonerEffector/Public/Cloner/CEClonerComponent.h b/Engine/Plugins/Experimental/ClonerEffector/Source/ClonerEffector/Public/Cloner/CEClonerComponent.h
--- a/Engine/Plugins/Experimental/ClonerEffector/Source/ClonerEffector/Public/Cloner/CEClonerComponent.h
+++ b/Engine/Plugins/Experimental/ClonerEffector/Source/ClonerEffector/Public/Cloner/CEClonerComponent.h
@@ -9,6 +9,7 @@
 
 class UCEClonerLayoutBase;
 class UMaterialInterface;
+struct FNiagaraMeshMaterialOverride;
 
 UCLASS(MinimalAPI, Within=CEClonerActor, HideCategories=(Niagara))
 class UCEClonerComponent : public UNiagaraComponent
@@ -24,6 +25,17 @@ public:
 	DECLARE_MULTICAST_DELEGATE_TwoParams(FOnClonerSystemLoaded, UCEClonerComponent* /** ClonerComponent */, UCEClonerLayoutBase* /** InLayout */)
 	static FOnClonerSystemLoaded OnClonerSystemLoaded;
 
+	/** Only materials transient or part of the content folder can be dirtied, engine or plugins cannot */
+	static bool IsMaterialDirtyable(const UMaterialInterface* InMaterial);
+
+	/** Check if material has niagara usage flag set */
+	static bool IsMaterialUsageFlagSet(const UMaterialInterface* InMaterial);
+
+#if WITH_EDITOR
+	/** Show material warning notification when missing niagara usage flag */
+	static void ShowMaterialWarning(int32 InMaterialCount);
+#endif
+
 	UCEClonerComponent();
 
 	/**
@@ -46,6 +58,10 @@ public:
 	/** Forces a refresh of the active system parameters in niagara store */
 	void RefreshUserParameters() const;
 
+	void SetUseOverrideMeshesMaterial(bool bInOverride);
+
+	void SetOverrideMeshesMaterial(UMaterialInterface* InMaterial);
+
 protected:
 	//~ Begin UObject
 	virtual void PostLoad() override;
@@ -102,8 +118,11 @@ protected:
 	/** Checks that all root static meshes are valid */
 	bool IsAllMergedMeshesValid() const;
 
-	/** Only materials transient or part of the content folder can be dirtied, engine or plugins cannot */
-	bool IsMaterialDirtyable(const UMaterialInterface* InMaterial) const;
+	/** Called when material override properties are changed */
+	void OnOverrideMeshesMaterialChanged();
+
+	/** Get the niagara materials override if enabled */
+	TArray<FNiagaraMeshMaterialOverride> GetOverrideMeshesMaterials() const;
 
 	/** Attachment tree view */
 	UPROPERTY(Transient, NonTransactional)
@@ -116,6 +135,12 @@ protected:
 	/** Attachment items that are dirty and need an update */
 	TSet<FCEClonerAttachmentItem*> DirtyItemAttachments;
 
+	/** Override meshes material */
+	TWeakObjectPtr<UMaterialInterface> ClonerMeshesOverrideMaterial;
+
+	/** Use override meshes material */
+	bool bUseClonerMeshesOverrideMaterial = false;
+
 	/** Asset meshes needs update */
 	bool bClonerMeshesDirty = true;
 
diff --git a/Engine/Plugins/Experimental/ClonerEffector/Source/ClonerEffector/Public/Effector/CEEffectorActor.h b/Engine/Plugins/Experimental/ClonerEffector/Source/ClonerEffector/Public/Effector/CEEffectorActor.h
--- a/Engine/Plugins/Experimental/ClonerEffector/Source/ClonerEffector/Public/Effector/CEEffectorActor.h
+++ b/Engine/Plugins/Experimental/ClonerEffector/Source/ClonerEffector/Public/Effector/CEEffectorActor.h
@@ -461,7 +461,7 @@ protected:
 
 	/** Affected clones color passed over to material */
 	UPROPERTY(EditAnywhere, BlueprintReadWrite, Setter, Getter, Category="Effector")
-	FLinearColor Color = FLinearColor::White;
+	FLinearColor Color = FLinearColor::Red;
 
 	/** Type of effector to apply on cloners instances */
 	UPROPERTY(EditAnywhere, BlueprintReadWrite, Setter="SetType", Getter="GetType", Category="Type")
