MotionDesign : ClonerEffector
- Fix changing a material property will update the cloner meshes instead of needing a forced refresh to update

diff --git a/Engine/Plugins/Experimental/ClonerEffector/Source/ClonerEffector/Private/Cloner/CEClonerComponent.cpp b/Engine/Plugins/Experimental/ClonerEffector/Source/ClonerEffector/Private/Cloner/CEClonerComponent.cpp
--- a/Engine/Plugins/Experimental/ClonerEffector/Source/ClonerEffector/Private/Cloner/CEClonerComponent.cpp
+++ b/Engine/Plugins/Experimental/ClonerEffector/Source/ClonerEffector/Private/Cloner/CEClonerComponent.cpp
@@ -441,6 +441,11 @@ void UCEClonerComponent::OnActorPropertyChanged(UObject* InObject, FPropertyChan
 
 void UCEClonerComponent::OnMaterialChanged(UObject* InObject)
 {
+	if (!IsValid(InObject))
+	{
+		return;
+	}
+
 	const AActor* ClonerActor = GetOwner();
 
 	if (!ClonerActor)
@@ -448,109 +453,118 @@ void UCEClonerComponent::OnMaterialChanged(UObject* InObject)
 		return;
 	}
 
-	if (AActor* ActorChanged = Cast<AActor>(InObject))
+	AActor* ActorChanged = Cast<AActor>(InObject);
+	ActorChanged = ActorChanged ? ActorChanged : InObject->GetTypedOuter<AActor>();
+
+	if (!ActorChanged)
 	{
-		if (FCEClonerAttachmentItem* Item = ClonerTree.ItemAttachmentMap.Find(ActorChanged))
-		{
-			int32 MatIdx = 0;
-			TArray<TWeakObjectPtr<UMaterialInterface>> NewMaterials;
-			bool bMaterialChanged = false;
+		return;
+	}
 
-			TArray<UDynamicMeshComponent*> DynamicMeshComponents;
-			ActorChanged->GetComponents(DynamicMeshComponents, false);
-			for (const UDynamicMeshComponent* DynamicMeshComponent : DynamicMeshComponents)
-			{
-				for (UMaterialInterface* Material : DynamicMeshComponent->GetMaterials())
-				{
-					if (!Item->BakedMaterials.IsValidIndex(MatIdx) || Item->BakedMaterials[MatIdx] != Material)
-					{
-						bMaterialChanged = true;
-					}
-					NewMaterials.Add(Material);
-					MatIdx++;
-				}
-			}
+	FCEClonerAttachmentItem* AttachmentItem = ClonerTree.ItemAttachmentMap.Find(ActorChanged);
+
+	if (!AttachmentItem)
+	{
+		return;
+	}
+
+	int32 MatIdx = 0;
+	TArray<TWeakObjectPtr<UMaterialInterface>> NewMaterials;
+	bool bMaterialChanged = false;
 
-			TArray<USkeletalMeshComponent*> SkeletalMeshComponents;
-			ActorChanged->GetComponents(SkeletalMeshComponents, false);
-			for (const USkeletalMeshComponent* SkeletalMeshComponent : SkeletalMeshComponents)
+	TArray<UDynamicMeshComponent*> DynamicMeshComponents;
+	ActorChanged->GetComponents(DynamicMeshComponents, false);
+	for (const UDynamicMeshComponent* DynamicMeshComponent : DynamicMeshComponents)
+	{
+		for (UMaterialInterface* Material : DynamicMeshComponent->GetMaterials())
+		{
+			if (!AttachmentItem->BakedMaterials.IsValidIndex(MatIdx) || AttachmentItem->BakedMaterials[MatIdx] != Material)
 			{
-				for (UMaterialInterface* Material : SkeletalMeshComponent->GetMaterials())
-				{
-					if (!Item->BakedMaterials.IsValidIndex(MatIdx) || Item->BakedMaterials[MatIdx] != Material)
-					{
-						bMaterialChanged = true;
-					}
-					NewMaterials.Add(Material);
-					MatIdx++;
-				}
+				bMaterialChanged = true;
 			}
+			NewMaterials.Add(Material);
+			MatIdx++;
+		}
+	}
 
-			TArray<UBrushComponent*> BrushComponents;
-			ActorChanged->GetComponents(BrushComponents, false);
-			for (const UBrushComponent* BrushComponent : BrushComponents)
+	TArray<USkeletalMeshComponent*> SkeletalMeshComponents;
+	ActorChanged->GetComponents(SkeletalMeshComponents, false);
+	for (const USkeletalMeshComponent* SkeletalMeshComponent : SkeletalMeshComponents)
+	{
+		for (UMaterialInterface* Material : SkeletalMeshComponent->GetMaterials())
+		{
+			if (!AttachmentItem->BakedMaterials.IsValidIndex(MatIdx) || AttachmentItem->BakedMaterials[MatIdx] != Material)
 			{
-				for (int32 Idx = 0; Idx < BrushComponent->GetNumMaterials(); Idx++)
-				{
-					if (!Item->BakedMaterials.IsValidIndex(MatIdx) || Item->BakedMaterials[MatIdx] != BrushComponent->GetMaterial(Idx))
-					{
-						bMaterialChanged = true;
-					}
-					NewMaterials.Add(BrushComponent->GetMaterial(Idx));
-					MatIdx++;
-				}
+				bMaterialChanged = true;
 			}
+			NewMaterials.Add(Material);
+			MatIdx++;
+		}
+	}
 
-			TArray<UProceduralMeshComponent*> ProceduralMeshComponents;
-			ActorChanged->GetComponents(ProceduralMeshComponents, false);
-			for (const UProceduralMeshComponent* ProceduralMeshComponent : ProceduralMeshComponents)
+	TArray<UBrushComponent*> BrushComponents;
+	ActorChanged->GetComponents(BrushComponents, false);
+	for (const UBrushComponent* BrushComponent : BrushComponents)
+	{
+		for (int32 Idx = 0; Idx < BrushComponent->GetNumMaterials(); Idx++)
+		{
+			if (!AttachmentItem->BakedMaterials.IsValidIndex(MatIdx) || AttachmentItem->BakedMaterials[MatIdx] != BrushComponent->GetMaterial(Idx))
 			{
-				if (ProceduralMeshComponent->GetNumSections() == 0)
-				{
-					continue;
-				}
-				for (UMaterialInterface* Material : ProceduralMeshComponent->GetMaterials())
-				{
-					if (!Item->BakedMaterials.IsValidIndex(MatIdx) || Item->BakedMaterials[MatIdx] != Material)
-					{
-						bMaterialChanged = true;
-					}
-					NewMaterials.Add(Material);
-					MatIdx++;
-				}
+				bMaterialChanged = true;
 			}
+			NewMaterials.Add(BrushComponent->GetMaterial(Idx));
+			MatIdx++;
+		}
+	}
 
-			TArray<UStaticMeshComponent*> StaticMeshComponents;
-			ActorChanged->GetComponents(StaticMeshComponents, false);
-			for (const UStaticMeshComponent* StaticMeshComponent : StaticMeshComponents)
+	TArray<UProceduralMeshComponent*> ProceduralMeshComponents;
+	ActorChanged->GetComponents(ProceduralMeshComponents, false);
+	for (const UProceduralMeshComponent* ProceduralMeshComponent : ProceduralMeshComponents)
+	{
+		if (ProceduralMeshComponent->GetNumSections() == 0)
+		{
+			continue;
+		}
+		for (UMaterialInterface* Material : ProceduralMeshComponent->GetMaterials())
+		{
+			if (!AttachmentItem->BakedMaterials.IsValidIndex(MatIdx) || AttachmentItem->BakedMaterials[MatIdx] != Material)
 			{
-				for (UMaterialInterface* Material : StaticMeshComponent->GetMaterials())
-				{
-					if (!Item->BakedMaterials.IsValidIndex(MatIdx) || Item->BakedMaterials[MatIdx] != Material)
-					{
-						bMaterialChanged = true;
-					}
-					NewMaterials.Add(Material);
-					MatIdx++;
-				}
+				bMaterialChanged = true;
 			}
+			NewMaterials.Add(Material);
+			MatIdx++;
+		}
+	}
 
-			if (bMaterialChanged)
+	TArray<UStaticMeshComponent*> StaticMeshComponents;
+	ActorChanged->GetComponents(StaticMeshComponents, false);
+	for (const UStaticMeshComponent* StaticMeshComponent : StaticMeshComponents)
+	{
+		for (UMaterialInterface* Material : StaticMeshComponent->GetMaterials())
+		{
+			if (!AttachmentItem->BakedMaterials.IsValidIndex(MatIdx) || AttachmentItem->BakedMaterials[MatIdx] != Material)
 			{
-				UE_LOG(LogCEClonerComponent, Log, TEXT("%s : Detected material change for %s"), *ClonerActor->GetActorNameOrLabel(), *ActorChanged->GetActorNameOrLabel());
+				bMaterialChanged = true;
+			}
+			NewMaterials.Add(Material);
+			MatIdx++;
+		}
+	}
 
-				if (NewMaterials.Num() == Item->BakedMaterials.Num())
-				{
-					Item->BakedMaterials = NewMaterials;
-				}
-				else
-				{
-					Item->MeshStatus = ECEClonerAttachmentStatus::Outdated;
-				}
+	if (bMaterialChanged)
+	{
+		UE_LOG(LogCEClonerComponent, Log, TEXT("%s : Detected material change for %s"), *ClonerActor->GetActorNameOrLabel(), *ActorChanged->GetActorNameOrLabel());
 
-				InvalidateBakedStaticMesh(ActorChanged);
-			}
+		if (NewMaterials.Num() == AttachmentItem->BakedMaterials.Num())
+		{
+			AttachmentItem->BakedMaterials = NewMaterials;
+		}
+		else
+		{
+			AttachmentItem->MeshStatus = ECEClonerAttachmentStatus::Outdated;
 		}
+
+		InvalidateBakedStaticMesh(ActorChanged);
 	}
 }
 
@@ -1472,16 +1486,16 @@ void UCEClonerComponent::UpdateClonerMeshes()
 				}
 			}
 		}
+
+		bClonerMeshesDirty = !DirtyItemAttachments.IsEmpty();
+
+		UE_LOG(LogCEClonerComponent, Log, TEXT("%s : Cloner mesh updated %i"), *ClonerActor->GetActorNameOrLabel(), ClonerTree.MergedBakedMeshes.Num())
 	}
 
 #if WITH_EDITORONLY_DATA
 	MeshRenderer->OnMeshChanged();
 #endif
 
-	bClonerMeshesDirty = !DirtyItemAttachments.IsEmpty();
-
-	UE_LOG(LogCEClonerComponent, Log, TEXT("%s : Cloner mesh updated %i"), *ClonerActor->GetActorNameOrLabel(), ClonerTree.MergedBakedMeshes.Num());
-
 	OnClonerMeshUpdated.Broadcast(this);
 }
 
