MotionDesign : ClonerEffector
- Add support for invert effector for all types and modes

#virtualized

diff --git a/Engine/Plugins/Experimental/ClonerEffector/Source/ClonerEffector/Private/Effector/CEEffectorActor.cpp b/Engine/Plugins/Experimental/ClonerEffector/Source/ClonerEffector/Private/Effector/CEEffectorActor.cpp
--- a/Engine/Plugins/Experimental/ClonerEffector/Source/ClonerEffector/Private/Effector/CEEffectorActor.cpp
+++ b/Engine/Plugins/Experimental/ClonerEffector/Source/ClonerEffector/Private/Effector/CEEffectorActor.cpp
@@ -121,6 +121,7 @@ TCEPropertyChangeDispatcher<ACEEffectorActor> ACEEffectorActor::PropertyChangeDi
 	{ GET_MEMBER_NAME_CHECKED(ACEEffectorActor, Color), &ACEEffectorActor::OnColorChanged },
 	/** Type */
 	{ GET_MEMBER_NAME_CHECKED(ACEEffectorActor, Type), &ACEEffectorActor::OnTypeChanged },
+	{ GET_MEMBER_NAME_CHECKED(ACEEffectorActor, bInvertType), &ACEEffectorActor::OnMagnitudeChanged },
 	{ GET_MEMBER_NAME_CHECKED(ACEEffectorActor, Easing), &ACEEffectorActor::OnEasingChanged },
 	{ GET_MEMBER_NAME_CHECKED(ACEEffectorActor, OuterRadius), &ACEEffectorActor::OnSphereChanged },
 	{ GET_MEMBER_NAME_CHECKED(ACEEffectorActor, InnerRadius), &ACEEffectorActor::OnSphereChanged },
@@ -364,6 +365,17 @@ void ACEEffectorActor::SetPlaneSpacing(float InSpacing)
 	OnPlaneChanged();
 }
 
+void ACEEffectorActor::SetInvertType(bool bInInvert)
+{
+	if (bInvertType == bInInvert)
+	{
+		return;
+	}
+
+	bInvertType = bInInvert;
+	OnMagnitudeChanged();
+}
+
 void ACEEffectorActor::SetOffset(const FVector& InOffset)
 {
 	if (InOffset.Equals(Offset))
@@ -892,7 +904,7 @@ void ACEEffectorActor::OnPlaneChanged()
 void ACEEffectorActor::OnMagnitudeChanged()
 {
 	const float EffectorMagnitude = bEnabled ? GetMagnitude() : 0.f;
-	ChannelData.Magnitude = EffectorMagnitude;
+	ChannelData.Magnitude = bInvertType ? -EffectorMagnitude : EffectorMagnitude;
 }
 
 void ACEEffectorActor::OnTransformOptionsChanged()
diff --git a/Engine/Plugins/Experimental/ClonerEffector/Source/ClonerEffector/Public/Effector/CEEffectorActor.h b/Engine/Plugins/Experimental/ClonerEffector/Source/ClonerEffector/Public/Effector/CEEffectorActor.h
--- a/Engine/Plugins/Experimental/ClonerEffector/Source/ClonerEffector/Public/Effector/CEEffectorActor.h
+++ b/Engine/Plugins/Experimental/ClonerEffector/Source/ClonerEffector/Public/Effector/CEEffectorActor.h
@@ -36,6 +36,15 @@ public:
 		return Type;
 	}
 
+	UFUNCTION(BlueprintCallable, Category="Effector")
+	CLONEREFFECTOR_API void SetInvertType(bool bInInvert);
+
+	UFUNCTION(BlueprintPure, Category="Effector")
+	bool GetInvertType() const
+	{
+		return bInvertType;
+	}
+
 	UFUNCTION(BlueprintCallable, Category="Effector")
 	CLONEREFFECTOR_API void SetEasing(ECEClonerEasing InEasing);
 
@@ -467,6 +476,10 @@ protected:
 	UPROPERTY(EditAnywhere, BlueprintReadWrite, Setter="SetType", Getter="GetType", Category="Type")
 	ECEClonerEffectorType Type = ECEClonerEffectorType::Sphere;
 
+	/** Invert the type effect, instead of affecting the inside of a zone, will affect the outside */
+	UPROPERTY(EditAnywhere, BlueprintReadWrite, Setter="SetInvertType", Getter="GetInvertType", Category="Type")
+	bool bInvertType = false;
+
 	/** Weight easing function applied to lerp transforms */
 	UPROPERTY(EditAnywhere, BlueprintReadWrite, Setter="SetEasing", Getter="GetEasing", Category="Type", meta=(EditCondition="Type != ECEClonerEffectorType::Unbound", EditConditionHides))
 	ECEClonerEasing Easing = ECEClonerEasing::Linear;
