MotionDesign : ClonerEffector
- Adding radial/spiral effector type support with visualizer
- Fixing rotation quaternion issue when effector and cloner are both rotated

#virtualized

diff --git a/Engine/Plugins/Experimental/ClonerEffector/Source/ClonerEffector/Private/Effector/CEEffectorActor.cpp b/Engine/Plugins/Experimental/ClonerEffector/Source/ClonerEffector/Private/Effector/CEEffectorActor.cpp
--- a/Engine/Plugins/Experimental/ClonerEffector/Source/ClonerEffector/Private/Effector/CEEffectorActor.cpp
+++ b/Engine/Plugins/Experimental/ClonerEffector/Source/ClonerEffector/Private/Effector/CEEffectorActor.cpp
@@ -3,14 +3,13 @@
 #include "Effector/CEEffectorActor.h"
 
 #include "Cloner/CEClonerActor.h"
+#include "Components/ArrowComponent.h"
 #include "Components/BoxComponent.h"
 #include "Components/SphereComponent.h"
 #include "Effector/CEEffectorComponent.h"
 #include "Math/Vector.h"
 #include "Subsystems/CEEffectorSubsystem.h"
 
-DEFINE_LOG_CATEGORY_STATIC(LogCEEffectorActor, Log, All);
-
 ACEEffectorActor::FOnEffectorIdentifierChanged ACEEffectorActor::OnEffectorRefreshClonerDelegate;
 
 ACEEffectorActor::ACEEffectorActor()
@@ -76,7 +75,7 @@ ACEEffectorActor::ACEEffectorActor()
 
 	// Plane
 	InnerPlaneComponent = CreateDefaultSubobject<UBoxComponent>(TEXT("AvaInnerPlaneComponent"));
-	InnerPlaneComponent->ShapeColor = FColor::Blue;
+	InnerPlaneComponent->ShapeColor = FColor::Red;
 	InnerPlaneComponent->SetLineThickness(VisualizerThickness);
 	InnerPlaneComponent->SetBoxExtent(InnerExtent);
 	InnerPlaneComponent->SetHiddenInGame(true);
@@ -88,7 +87,7 @@ ACEEffectorActor::ACEEffectorActor()
 	InnerPlaneComponent->SetupAttachment(SceneComponent);
 
 	OuterPlaneComponent = CreateDefaultSubobject<UBoxComponent>(TEXT("AvaOuterPlaneComponent"));
-	OuterPlaneComponent->ShapeColor = FColor::Red;
+	OuterPlaneComponent->ShapeColor = FColor::Blue;
 	OuterPlaneComponent->SetLineThickness(VisualizerThickness);
 	OuterPlaneComponent->SetBoxExtent(OuterExtent);
 	OuterPlaneComponent->SetHiddenInGame(true);
@@ -99,6 +98,31 @@ ACEEffectorActor::ACEEffectorActor()
 	OuterPlaneComponent->bIsEditorOnly = false;
 	OuterPlaneComponent->SetupAttachment(SceneComponent);
 
+	// Radial
+	BeginRadialComponent = CreateDefaultSubobject<UArrowComponent>(TEXT("BeginRadialComponent"));
+	BeginRadialComponent->SetRelativeRotation(FRotator(0, -90, 0));
+	BeginRadialComponent->SetArrowLength(100);
+	BeginRadialComponent->SetArrowColor(FColor::Red);
+	BeginRadialComponent->SetHiddenInGame(true);
+#if WITH_EDITOR
+	// Do not show bounding box around cloner for better visibility
+	BeginRadialComponent->SetIsVisualizationComponent(true);
+#endif
+	BeginRadialComponent->bIsEditorOnly = false;
+	BeginRadialComponent->SetupAttachment(SceneComponent);
+
+	EndRadialComponent = CreateDefaultSubobject<UArrowComponent>(TEXT("EndRadialComponent"));
+	EndRadialComponent->SetRelativeRotation(FRotator(0, -90 + RadialAngle, 0));
+	EndRadialComponent->SetArrowLength(100);
+	EndRadialComponent->SetArrowColor(FColor::Blue);
+	EndRadialComponent->SetHiddenInGame(true);
+#if WITH_EDITOR
+	// Do not show bounding box around cloner for better visibility
+	EndRadialComponent->SetIsVisualizationComponent(true);
+#endif
+	EndRadialComponent->bIsEditorOnly = false;
+	EndRadialComponent->SetupAttachment(SceneComponent);
+
 	if (!IsTemplate())
 	{
 		SceneComponent->TransformUpdated.AddUObject(this, &ACEEffectorActor::OnEffectorTransformed);
@@ -128,6 +152,9 @@ TCEPropertyChangeDispatcher<ACEEffectorActor> ACEEffectorActor::PropertyChangeDi
 	{ GET_MEMBER_NAME_CHECKED(ACEEffectorActor, InnerExtent), &ACEEffectorActor::OnBoxChanged },
 	{ GET_MEMBER_NAME_CHECKED(ACEEffectorActor, OuterExtent), &ACEEffectorActor::OnBoxChanged },
 	{ GET_MEMBER_NAME_CHECKED(ACEEffectorActor, PlaneSpacing), &ACEEffectorActor::OnPlaneChanged },
+	{ GET_MEMBER_NAME_CHECKED(ACEEffectorActor, RadialAngle), &ACEEffectorActor::OnRadialChanged },
+	{ GET_MEMBER_NAME_CHECKED(ACEEffectorActor, RadialMinRadius), &ACEEffectorActor::OnRadialChanged },
+	{ GET_MEMBER_NAME_CHECKED(ACEEffectorActor, RadialMaxRadius), &ACEEffectorActor::OnRadialChanged },
 	/** Mode */
 	{ GET_MEMBER_NAME_CHECKED(ACEEffectorActor, Mode), &ACEEffectorActor::OnModeChanged },
 	{ GET_MEMBER_NAME_CHECKED(ACEEffectorActor, Offset), &ACEEffectorActor::OnTransformOptionsChanged },
@@ -365,6 +392,32 @@ void ACEEffectorActor::SetPlaneSpacing(float InSpacing)
 	OnPlaneChanged();
 }
 
+void ACEEffectorActor::SetRadialAngle(float InAngle)
+{
+	InAngle = FMath::Clamp(InAngle, 0, 360);
+
+	if (FMath::IsNearlyEqual(InAngle, RadialAngle))
+	{
+		return;
+	}
+
+	RadialAngle = InAngle;
+	OnRadialChanged();
+}
+
+void ACEEffectorActor::SetRadialMinRadius(float InRadius)
+{
+	InRadius = FMath::Max(0, InRadius);
+
+	if (FMath::IsNearlyEqual(InRadius, RadialMinRadius))
+	{
+		return;
+	}
+
+	RadialMinRadius = InRadius;
+	OnRadialChanged();
+}
+
 void ACEEffectorActor::SetInvertType(bool bInInvert)
 {
 	if (bInvertType == bInInvert)
@@ -376,6 +429,19 @@ void ACEEffectorActor::SetInvertType(bool bInInvert)
 	OnMagnitudeChanged();
 }
 
+void ACEEffectorActor::SetRadialMaxRadius(float InRadius)
+{
+	InRadius = FMath::Max(0, InRadius);
+
+	if (FMath::IsNearlyEqual(InRadius, RadialMaxRadius))
+	{
+		return;
+	}
+
+	RadialMaxRadius = InRadius;
+	OnRadialChanged();
+}
+
 void ACEEffectorActor::SetOffset(const FVector& InOffset)
 {
 	if (InOffset.Equals(Offset))
@@ -727,10 +793,13 @@ void ACEEffectorActor::SetColor(const FLinearColor& InColor)
 void ACEEffectorActor::OnEffectorTransformed(USceneComponent* InUpdatedComponent, EUpdateTransformFlags InUpdateTransformFlags, ETeleportType InTeleport)
 {
 	OnTransformChanged();
+
 	// Update when scaled or rotated
 	OnSphereChanged();
 	OnBoxChanged();
 	OnPlaneChanged();
+	OnRadialChanged();
+
 	// update if self
 	const AActor* InternalTargetActor = InternalTargetActorWeak.Get();
 	if (InternalTargetActor == this)
@@ -793,12 +862,20 @@ void ACEEffectorActor::OnTypeChanged()
 		OuterPlaneComponent->SetVisibility(Type == ECEClonerEffectorType::Plane);
 	}
 
+	// Radial
+	if (BeginRadialComponent && EndRadialComponent)
+	{
+		BeginRadialComponent->SetVisibility(Type == ECEClonerEffectorType::Radial);
+		EndRadialComponent->SetVisibility(Type == ECEClonerEffectorType::Radial);
+	}
+
 	ChannelData.Type = Type;
 
 	// Update type data
 	OnSphereChanged();
 	OnBoxChanged();
 	OnPlaneChanged();
+	OnRadialChanged();
 }
 
 void ACEEffectorActor::OnEffectorChanged()
@@ -877,7 +954,7 @@ void ACEEffectorActor::OnSphereChanged()
 
 void ACEEffectorActor::OnPlaneChanged()
 {
-	static const FVector PlaneAxis = -FVector::LeftVector;
+	static const FVector PlaneAxis = FVector::LeftVector;
 	const FVector InnerPlane = PlaneAxis * FVector(-PlaneSpacing/2);
 	const FVector OuterPlane = PlaneAxis * FVector(PlaneSpacing/2);
 
@@ -901,6 +978,25 @@ void ACEEffectorActor::OnPlaneChanged()
 	ChannelData.OuterExtent = FVector(PlaneSpacing);
 }
 
+void ACEEffectorActor::OnRadialChanged()
+{
+	RadialMinRadius = FMath::Min(RadialMinRadius, RadialMaxRadius);
+	RadialMaxRadius = FMath::Max(RadialMinRadius, RadialMaxRadius);
+	
+	if (BeginRadialComponent && EndRadialComponent)
+	{
+		EndRadialComponent->SetRelativeRotation(FRotator(0, -90 + RadialAngle, 0));
+	}
+
+	if (Type != ECEClonerEffectorType::Radial)
+	{
+		return;
+	}
+
+	ChannelData.InnerExtent = FVector::LeftVector;
+	ChannelData.OuterExtent = FVector(RadialAngle, RadialMinRadius, RadialMaxRadius);
+}
+
 void ACEEffectorActor::OnMagnitudeChanged()
 {
 	const float EffectorMagnitude = bEnabled ? GetMagnitude() : 0.f;
@@ -1033,6 +1129,13 @@ void ACEEffectorActor::OnVisualizerThicknessChanged()
 		InnerPlaneComponent->SetLineThickness(VisualizerThickness);
 		OuterPlaneComponent->SetLineThickness(VisualizerThickness);
 	}
+
+	if (BeginRadialComponent && EndRadialComponent)
+	{
+		const float ArrowScale = VisualizerThickness / 2;
+		BeginRadialComponent->SetRelativeScale3D(FVector(10, ArrowScale, ArrowScale));
+		EndRadialComponent->SetRelativeScale3D(FVector(10, ArrowScale, ArrowScale));
+	}
 }
 
 void ACEEffectorActor::OnVisualizerSpriteVisibleChanged()
diff --git a/Engine/Plugins/Experimental/ClonerEffector/Source/ClonerEffector/Public/CEClonerEffectorShared.h b/Engine/Plugins/Experimental/ClonerEffector/Source/ClonerEffector/Public/CEClonerEffectorShared.h
--- a/Engine/Plugins/Experimental/ClonerEffector/Source/ClonerEffector/Public/CEClonerEffectorShared.h
+++ b/Engine/Plugins/Experimental/ClonerEffector/Source/ClonerEffector/Public/CEClonerEffectorShared.h
@@ -2,7 +2,7 @@
 
 #pragma once
 
-#include "Components/SphereComponent.h"
+#include "Components/SceneComponent.h"
 #include "Engine/StaticMesh.h"
 #include "Materials/MaterialInterface.h"
 #include "NiagaraDataInterfaceArrayInt.h"
@@ -152,7 +152,9 @@ enum class ECEClonerEffectorType : uint8
 	/** Clones inside the box extent will be affected by the effector */
 	Box,
 	/** All clones will be affected by the effector with the same max weight */
-	Unbound
+	Unbound,
+	/** All clones within the angle range will be affected */
+	Radial
 };
 
 /** Enumerates the effector mode available */
diff --git a/Engine/Plugins/Experimental/ClonerEffector/Source/ClonerEffector/Public/Effector/CEEffectorActor.h b/Engine/Plugins/Experimental/ClonerEffector/Source/ClonerEffector/Public/Effector/CEEffectorActor.h
--- a/Engine/Plugins/Experimental/ClonerEffector/Source/ClonerEffector/Public/Effector/CEEffectorActor.h
+++ b/Engine/Plugins/Experimental/ClonerEffector/Source/ClonerEffector/Public/Effector/CEEffectorActor.h
@@ -9,6 +9,7 @@
 
 class UNiagaraDataChannelWriter;
 class ACEClonerActor;
+class UArrowComponent;
 class USphereComponent;
 class UBoxComponent;
 
@@ -108,6 +109,33 @@ public:
 		return PlaneSpacing;
 	}
 
+	UFUNCTION(BlueprintCallable, Category="Effector")
+	CLONEREFFECTOR_API void SetRadialAngle(float InAngle);
+
+	UFUNCTION(BlueprintPure, Category="Effector")
+	float GetRadialAngle() const
+	{
+		return RadialAngle;
+	}
+
+	UFUNCTION(BlueprintCallable, Category="Effector")
+	CLONEREFFECTOR_API void SetRadialMinRadius(float InRadius);
+
+	UFUNCTION(BlueprintPure, Category="Effector")
+	float GetRadialMinRadius() const
+	{
+		return RadialMinRadius;
+	}
+
+	UFUNCTION(BlueprintCallable, Category="Effector")
+	CLONEREFFECTOR_API void SetRadialMaxRadius(float InRadius);
+
+	UFUNCTION(BlueprintPure, Category="Effector")
+	float GetRadialMaxRadius() const
+	{
+		return RadialMaxRadius;
+	}
+
 	UFUNCTION(BlueprintCallable, Category="Effector")
 	CLONEREFFECTOR_API void SetOffset(const FVector& InOffset);
 
@@ -424,6 +452,9 @@ protected:
 	/** Update values for plane type effector */
 	void OnPlaneChanged();
 
+	/** Update values for radial type effector */
+	void OnRadialChanged();
+
 	/** Called when transform of effector has changed */
 	void OnTransformChanged();
 
@@ -504,6 +535,18 @@ protected:
 	UPROPERTY(EditAnywhere, BlueprintReadWrite, Setter="SetPlaneSpacing", Getter="GetPlaneSpacing", Category="Type", meta=(ClampMin="0", EditCondition="Type == ECEClonerEffectorType::Plane", EditConditionHides))
 	float PlaneSpacing = 200.f;
 
+	/** Radial angle in degree, everything within the angle will be affected */
+	UPROPERTY(EditAnywhere, BlueprintReadWrite, Setter, Getter, Category="Type", meta=(ClampMin="0", ClampMax="360", EditCondition="Type == ECEClonerEffectorType::Radial", EditConditionHides))
+	float RadialAngle = 180.f;
+
+	/** Minimum radius for the radial effect to be applied on clones, below clones will not be affected */
+	UPROPERTY(EditAnywhere, BlueprintReadWrite, Setter, Getter, Category="Type", meta=(ClampMin="0", EditCondition="Type == ECEClonerEffectorType::Radial", EditConditionHides))
+	float RadialMinRadius = 0.f;
+
+	/** Maximum radius for the radial effect to be applied on clones, above clones will not be affected */
+	UPROPERTY(EditAnywhere, BlueprintReadWrite, Setter, Getter, Category="Type", meta=(ClampMin="0", EditCondition="Type == ECEClonerEffectorType::Radial", EditConditionHides))
+	float RadialMaxRadius = 1000.f;
+
 	/** Mode of effector for each clones instances */
 	UPROPERTY(EditAnywhere, BlueprintReadWrite, Setter="SetMode", Getter="GetMode", Category="Mode")
 	ECEClonerEffectorMode Mode = ECEClonerEffectorMode::Default;
@@ -637,6 +680,14 @@ private:
 	UPROPERTY()
 	TObjectPtr<UBoxComponent> OuterPlaneComponent;
 
+	/** Radial components for begin angle */
+	UPROPERTY()
+	TObjectPtr<UArrowComponent> BeginRadialComponent;
+
+	/** Radial components for end angle */
+	UPROPERTY()
+	TObjectPtr<UArrowComponent> EndRadialComponent;
+
 	/** Internal cloners array, deprecated since it will be moved to the cloners and emptied out on post load */
 	UPROPERTY(NonTransactional)
 	TSet<TWeakObjectPtr<ACEClonerActor>> InternalCloners;
