MotionDesign : ClonerEffector
- Effector customization easing widget improved
- svg easing icons updated (color changed to white to use color and opacity in widget)

diff --git a/Engine/Plugins/Experimental/ClonerEffector/Source/ClonerEffectorEditor/Private/Effector/CEEditorEffectorDetailCustomization.cpp b/Engine/Plugins/Experimental/ClonerEffector/Source/ClonerEffectorEditor/Private/Effector/CEEditorEffectorDetailCustomization.cpp
--- a/Engine/Plugins/Experimental/ClonerEffector/Source/ClonerEffectorEditor/Private/Effector/CEEditorEffectorDetailCustomization.cpp
+++ b/Engine/Plugins/Experimental/ClonerEffector/Source/ClonerEffectorEditor/Private/Effector/CEEditorEffectorDetailCustomization.cpp
@@ -11,8 +11,11 @@
 #include "PropertyHandle.h"
 #include "Styles/CEEditorStyle.h"
 #include "Styling/SlateBrush.h"
+#include "Styling/SlateColor.h"
 #include "Widgets/Images/SImage.h"
 #include "Widgets/Input/SComboBox.h"
+#include "Widgets/Layout/SScaleBox.h"
+#include "Widgets/SBoxPanel.h"
 
 #define LOCTEXT_NAMESPACE "CEEditorEffectorDetailCustomization"
 
@@ -37,22 +40,22 @@ void FCEEditorEffectorDetailCustomization::CustomizeDetails(IDetailLayoutBuilder
 		FDetailWidgetRow& CustomWidget = EasingRow->CustomWidget();
 
 		CustomWidget.NameContent()
-		.HAlign(EHorizontalAlignment::HAlign_Fill)
 		.VAlign(EVerticalAlignment::VAlign_Fill)
 		[
 			EasingPropertyHandle->CreatePropertyNameWidget()
 		];
 
 		CustomWidget.ValueContent()
-		.HAlign(EHorizontalAlignment::HAlign_Fill)
 		.VAlign(EVerticalAlignment::VAlign_Fill)
 		[
 			SNew(SComboBox<FName>)
+			.ComboBoxStyle(&FCEEditorStyle::Get().GetWidgetStyle<FComboBoxStyle>(TEXT("ComboBox")))
 			.OptionsSource(&EasingNames)
 			.InitiallySelectedItem(GetCurrentEasingName())
 			.ToolTipText(LOCTEXT("EasingTooltip", "Easings sorted from most dramatic to least and specials at the end"))
 			.OnGenerateWidget(this, &FCEEditorEffectorDetailCustomization::OnGenerateEasingEntry)
 			.OnSelectionChanged(this, &FCEEditorEffectorDetailCustomization::OnSelectionChanged)
+			.ContentPadding(0.f)
 			.Content()
 			[
 				OnGenerateEasingEntry(NAME_None)
@@ -143,6 +146,7 @@ TSharedRef<SWidget> FCEEditorEffectorDetailCustomization::OnGenerateEasingEntry(
 {
 	TSharedPtr<SWidget> ImageWidget;
 	TSharedPtr<SWidget> TextWidget;
+	const TSharedPtr<SHorizontalBox> HorizontalWidget = SNew(SHorizontalBox).Visibility(EVisibility::Visible);
 
 	static const FVector2D ImageSizeClosed(16.f, 16.f);
 	static const FVector2D ImageSizeOpened(32.f, 32.f);
@@ -151,37 +155,58 @@ TSharedRef<SWidget> FCEEditorEffectorDetailCustomization::OnGenerateEasingEntry(
 	if (InName == NAME_None)
 	{
 		SAssignNew(ImageWidget, SImage)
+		.ColorAndOpacity(FAppStyle::GetSlateColor("SelectionColor"))
 		.DesiredSizeOverride(ImageSizeClosed)
 		.Image(this, &FCEEditorEffectorDetailCustomization::GetEasingImage, InName);
 
 		SAssignNew(TextWidget, STextBlock)
+		.Justification(ETextJustify::Center)
 		.Text(this, &FCEEditorEffectorDetailCustomization::GetEasingText, InName);
 	}
 	else
 	{
+		// We need to switch color on hover to avoid having the selected color equals the image color
+		const TSharedPtr<SWidget> HoverWidget = HorizontalWidget;
+
 		SAssignNew(ImageWidget, SImage)
+		.ColorAndOpacity_Static(&FCEEditorEffectorDetailCustomization::GetImageColorAndOpacity, HoverWidget)
 		.DesiredSizeOverride(ImageSizeOpened)
 		.Image(GetEasingImage(InName));
 
 		SAssignNew(TextWidget, STextBlock)
+		.Justification(ETextJustify::Center)
 		.Text(GetEasingText(InName));
 	}
 
-	return SNew(SHorizontalBox)
-		+ SHorizontalBox::Slot()
+	// Make inner widgets hit test invisible and only horizontal is hit testable
+
+	HorizontalWidget->AddSlot()
 		.AutoWidth()
 		.Padding(2.f)
 		[
-			ImageWidget.ToSharedRef()
-		]
-		+ SHorizontalBox::Slot()
+			SNew(SScaleBox)
+			.Visibility(EVisibility::HitTestInvisible)
+			.Stretch(EStretch::UserSpecified)
+			.UserSpecifiedScale(1.5)
+			[
+				ImageWidget.ToSharedRef()
+			]
+		];
+
+	HorizontalWidget->AddSlot()
 		.FillWidth(1.f)
-		.Padding(5.f, 2.f)
-		.HAlign(EHorizontalAlignment::HAlign_Left)
+		.Padding(8.f, 2.f)
+		.HAlign(EHorizontalAlignment::HAlign_Fill)
 		.VAlign(EVerticalAlignment::VAlign_Center)
 		[
-			TextWidget.ToSharedRef()
+			SNew(SBox)
+			.Visibility(EVisibility::HitTestInvisible)
+			[
+				TextWidget.ToSharedRef()
+			]
 		];
+
+	return HorizontalWidget.ToSharedRef();
 }
 
 void FCEEditorEffectorDetailCustomization::OnSelectionChanged(FName InSelection, ESelectInfo::Type InSelectInfo) const
@@ -247,4 +272,14 @@ FText FCEEditorEffectorDetailCustomization::GetEasingText(FName InName) const
 	return FText::GetEmpty();
 }
 
+FSlateColor FCEEditorEffectorDetailCustomization::GetImageColorAndOpacity(const TSharedPtr<SWidget> InWidget)
+{
+	if (InWidget.IsValid())
+	{
+		return InWidget->IsHovered() ? FLinearColor::White : FAppStyle::GetSlateColor("SelectionColor");
+	}
+
+	return FLinearColor::White;
+}
+
 #undef LOCTEXT_NAMESPACE
diff --git a/Engine/Plugins/Experimental/ClonerEffector/Source/ClonerEffectorEditor/Private/Effector/CEEditorEffectorDetailCustomization.h b/Engine/Plugins/Experimental/ClonerEffector/Source/ClonerEffectorEditor/Private/Effector/CEEditorEffectorDetailCustomization.h
--- a/Engine/Plugins/Experimental/ClonerEffector/Source/ClonerEffectorEditor/Private/Effector/CEEditorEffectorDetailCustomization.h
+++ b/Engine/Plugins/Experimental/ClonerEffector/Source/ClonerEffectorEditor/Private/Effector/CEEditorEffectorDetailCustomization.h
@@ -10,7 +10,9 @@
 class IPropertyHandle;
 class SWidget;
 class UEnum;
+struct FLinearColor;
 struct FSlateBrush;
+struct FSlateColor;
 
 /* Used to customize effector actor properties in details panel */
 class FCEEditorEffectorDetailCustomization : public IDetailCustomization
@@ -31,6 +33,8 @@ public:
 	//~ End IDetailCustomization
 
 protected:
+	static FSlateColor GetImageColorAndOpacity(const TSharedPtr<SWidget> InWidget);
+
 	void RegisterCustomSections() const;
 
 	void PopulateEasingInfos();
diff --git a/Engine/Plugins/Experimental/ClonerEffector/Source/ClonerEffectorEditor/Private/Styles/CEEditorStyle.cpp b/Engine/Plugins/Experimental/ClonerEffector/Source/ClonerEffectorEditor/Private/Styles/CEEditorStyle.cpp
--- a/Engine/Plugins/Experimental/ClonerEffector/Source/ClonerEffectorEditor/Private/Styles/CEEditorStyle.cpp
+++ b/Engine/Plugins/Experimental/ClonerEffector/Source/ClonerEffectorEditor/Private/Styles/CEEditorStyle.cpp
@@ -5,6 +5,7 @@
 #include "CEClonerEffectorShared.h"
 #include "Interfaces/IPluginManager.h"
 #include "Misc/Paths.h"
+#include "Styling/AppStyle.h"
 #include "Styling/SlateStyleMacros.h"
 #include "Styling/StyleColors.h"
 #include "Styling/SlateTypes.h"
@@ -38,6 +39,11 @@ FCEEditorStyle::FCEEditorStyle()
 		}
 	}
 
+	// Combobox
+	FComboBoxStyle ComboBoxStyle = FAppStyle::Get().GetWidgetStyle<FComboBoxStyle>(TEXT("ComboBox"));
+	ComboBoxStyle.SetMenuRowPadding(FMargin(0));
+	Set("ComboBox", ComboBoxStyle);
+
 	FSlateStyleRegistry::RegisterSlateStyle(*this);
 }
 
