- FIx for member variable being changed

diff --git a/Engine/Plugins/Experimental/ClonerEffector/Source/ClonerEffector/Private/Cloner/Layouts/CEClonerSplineLayout.cpp b/Engine/Plugins/Experimental/ClonerEffector/Source/ClonerEffector/Private/Cloner/Layouts/CEClonerSplineLayout.cpp
--- a/Engine/Plugins/Experimental/ClonerEffector/Source/ClonerEffector/Private/Cloner/Layouts/CEClonerSplineLayout.cpp
+++ b/Engine/Plugins/Experimental/ClonerEffector/Source/ClonerEffector/Private/Cloner/Layouts/CEClonerSplineLayout.cpp
@@ -136,7 +136,7 @@ void UCEClonerSplineLayout::OnLayoutParametersChanged(UCEClonerComponent* InComp
 	InComponent->SetBoolParameter(TEXT("MeshOrientAxisEnable"), bOrientMesh);
 
 	// unbind
-	SplineDI->Source = nullptr;
+	SplineDI->SoftSourceActor = nullptr;
 
 	if (USplineComponent* SplineComponent = SplineComponentWeak.Get())
     {
@@ -151,7 +151,7 @@ void UCEClonerSplineLayout::OnLayoutParametersChanged(UCEClonerComponent* InComp
 	{
 		if (USplineComponent* SplineComponent = SplineActor->FindComponentByClass<USplineComponent>())
 		{
-			SplineDI->Source = SplineActor;
+			SplineDI->SoftSourceActor = SplineActor;
 			SplineComponentWeak = SplineComponent;
 
 			SplineComponent->TransformUpdated.RemoveAll(this);
