MotionDesign : ClonerEffector
- Adding new effector type torus
- Moving some operations in niagara functions
- Replacing all visualizer components by dynamic mesh components
- Registering effector to subsystem directly instead of waiting for cloner

#virtualized

diff --git a/Engine/Plugins/Experimental/ClonerEffector/Source/ClonerEffector/ClonerEffector.Build.cs b/Engine/Plugins/Experimental/ClonerEffector/Source/ClonerEffector/ClonerEffector.Build.cs
--- a/Engine/Plugins/Experimental/ClonerEffector/Source/ClonerEffector/ClonerEffector.Build.cs
+++ b/Engine/Plugins/Experimental/ClonerEffector/Source/ClonerEffector/ClonerEffector.Build.cs
@@ -15,17 +15,18 @@ public class ClonerEffector : ModuleRules
 				"CoreUObject",
 				"Engine",
 				"Niagara",
-				"NiagaraCore"
+				"NiagaraCore",
+				"StructUtils"
 			}
 		);
 
 		PrivateDependencyModuleNames.AddRange(
 			new string[]
 			{
+				"DynamicMesh",
 				"GeometryCore",
 				"GeometryFramework",
 				"GeometryScriptingCore",
-				"DynamicMesh",
 				"ProceduralMeshComponent",
 				"StaticMeshDescription"
 			}
diff --git a/Engine/Plugins/Experimental/ClonerEffector/Source/ClonerEffector/Private/Effector/CEEffectorActor.cpp b/Engine/Plugins/Experimental/ClonerEffector/Source/ClonerEffector/Private/Effector/CEEffectorActor.cpp
--- a/Engine/Plugins/Experimental/ClonerEffector/Source/ClonerEffector/Private/Effector/CEEffectorActor.cpp
+++ b/Engine/Plugins/Experimental/ClonerEffector/Source/ClonerEffector/Private/Effector/CEEffectorActor.cpp
@@ -3,12 +3,14 @@
 #include "Effector/CEEffectorActor.h"
 
 #include "Cloner/CEClonerActor.h"
-#include "Components/ArrowComponent.h"
-#include "Components/BoxComponent.h"
-#include "Components/SphereComponent.h"
+#include "Components/DynamicMeshComponent.h"
 #include "Effector/CEEffectorComponent.h"
+#include "GeometryScript/MeshPrimitiveFunctions.h"
+#include "Materials/MaterialInstanceDynamic.h"
 #include "Math/Vector.h"
+#include "PropertyBag.h"
 #include "Subsystems/CEEffectorSubsystem.h"
+#include "UObject/ConstructorHelpers.h"
 
 ACEEffectorActor::FOnEffectorIdentifierChanged ACEEffectorActor::OnEffectorRefreshClonerDelegate;
 
@@ -23,109 +25,53 @@ ACEEffectorActor::ACEEffectorActor()
 #endif
 	SetRootComponent(SceneComponent);
 
-	// Sphere
-	InnerSphereComponent = CreateDefaultSubobject<USphereComponent>(TEXT("AvaInnerSphereComponent"));;
-	InnerSphereComponent->ShapeColor = FColor::Blue;
-	InnerSphereComponent->SetLineThickness(VisualizerThickness);
-	InnerSphereComponent->SetSphereRadius(InnerRadius);
-	InnerSphereComponent->SetHiddenInGame(true);
+	// Inner Visualizer Component
+	InnerVisualizerComponent = CreateDefaultSubobject<UDynamicMeshComponent>(TEXT("InnerVisualizerComponent"));
+	InnerVisualizerComponent->SetHiddenInGame(true);
+	InnerVisualizerComponent->SetTranslucentSortPriority(0);
 #if WITH_EDITOR
 	// Do not show bounding box around cloner for better visibility
-	InnerSphereComponent->SetIsVisualizationComponent(true);
+	InnerVisualizerComponent->SetIsVisualizationComponent(true);
 #endif
-	InnerSphereComponent->bIsEditorOnly = false;
-	InnerSphereComponent->SetupAttachment(SceneComponent);
-
-	OuterSphereComponent = CreateDefaultSubobject<USphereComponent>(TEXT("AvaOuterSphereComponent"));
-	OuterSphereComponent->ShapeColor = FColor::Red;
-	OuterSphereComponent->SetLineThickness(VisualizerThickness);
-	OuterSphereComponent->SetSphereRadius(OuterRadius);
-	OuterSphereComponent->SetHiddenInGame(true);
-#if WITH_EDITOR
-	// Do not show bounding box around cloner for better visibility
-	OuterSphereComponent->SetIsVisualizationComponent(true);
-#endif
-	OuterSphereComponent->bIsEditorOnly = false;
-	OuterSphereComponent->SetupAttachment(SceneComponent);
-
-	// Box
-	InnerBoxComponent = CreateDefaultSubobject<UBoxComponent>(TEXT("AvaInnerBoxComponent"));
-	InnerBoxComponent->ShapeColor = FColor::Blue;
-	InnerBoxComponent->SetLineThickness(VisualizerThickness);
-	InnerBoxComponent->SetBoxExtent(InnerExtent);
-	InnerBoxComponent->SetHiddenInGame(true);
-#if WITH_EDITOR
-	// Do not show bounding box around cloner for better visibility
-	InnerBoxComponent->SetIsVisualizationComponent(true);
-#endif
-	InnerBoxComponent->bIsEditorOnly = false;
-	InnerBoxComponent->SetupAttachment(SceneComponent);
-
-	OuterBoxComponent = CreateDefaultSubobject<UBoxComponent>(TEXT("AvaOuterBoxComponent"));
-	OuterBoxComponent->ShapeColor = FColor::Red;
-	OuterBoxComponent->SetLineThickness(VisualizerThickness);
-	OuterBoxComponent->SetBoxExtent(OuterExtent);
-	OuterBoxComponent->SetHiddenInGame(true);
-#if WITH_EDITOR
-	// Do not show bounding box around cloner for better visibility
-	OuterBoxComponent->SetIsVisualizationComponent(true);
-#endif
-	OuterBoxComponent->bIsEditorOnly = false;
-	OuterBoxComponent->SetupAttachment(SceneComponent);
-
-	// Plane
-	InnerPlaneComponent = CreateDefaultSubobject<UBoxComponent>(TEXT("AvaInnerPlaneComponent"));
-	InnerPlaneComponent->ShapeColor = FColor::Red;
-	InnerPlaneComponent->SetLineThickness(VisualizerThickness);
-	InnerPlaneComponent->SetBoxExtent(InnerExtent);
-	InnerPlaneComponent->SetHiddenInGame(true);
-#if WITH_EDITOR
-	// Do not show bounding box around cloner for better visibility
-	InnerPlaneComponent->SetIsVisualizationComponent(true);
-#endif
-	InnerPlaneComponent->bIsEditorOnly = false;
-	InnerPlaneComponent->SetupAttachment(SceneComponent);
-
-	OuterPlaneComponent = CreateDefaultSubobject<UBoxComponent>(TEXT("AvaOuterPlaneComponent"));
-	OuterPlaneComponent->ShapeColor = FColor::Blue;
-	OuterPlaneComponent->SetLineThickness(VisualizerThickness);
-	OuterPlaneComponent->SetBoxExtent(OuterExtent);
-	OuterPlaneComponent->SetHiddenInGame(true);
-#if WITH_EDITOR
-	// Do not show bounding box around cloner for better visibility
-	OuterPlaneComponent->SetIsVisualizationComponent(true);
-#endif
-	OuterPlaneComponent->bIsEditorOnly = false;
-	OuterPlaneComponent->SetupAttachment(SceneComponent);
-
-	// Radial
-	BeginRadialComponent = CreateDefaultSubobject<UArrowComponent>(TEXT("BeginRadialComponent"));
-	BeginRadialComponent->SetRelativeRotation(FRotator(0, -90, 0));
-	BeginRadialComponent->SetArrowLength(100);
-	BeginRadialComponent->SetArrowColor(FColor::Red);
-	BeginRadialComponent->SetHiddenInGame(true);
-#if WITH_EDITOR
-	// Do not show bounding box around cloner for better visibility
-	BeginRadialComponent->SetIsVisualizationComponent(true);
-#endif
-	BeginRadialComponent->bIsEditorOnly = false;
-	BeginRadialComponent->SetupAttachment(SceneComponent);
-
-	EndRadialComponent = CreateDefaultSubobject<UArrowComponent>(TEXT("EndRadialComponent"));
-	EndRadialComponent->SetRelativeRotation(FRotator(0, -90 + RadialAngle, 0));
-	EndRadialComponent->SetArrowLength(100);
-	EndRadialComponent->SetArrowColor(FColor::Blue);
-	EndRadialComponent->SetHiddenInGame(true);
+	InnerVisualizerComponent->bIsEditorOnly = false;
+	InnerVisualizerComponent->SetupAttachment(SceneComponent);
+
+	// Outer Visualizer Component
+	OuterVisualizerComponent = CreateDefaultSubobject<UDynamicMeshComponent>(TEXT("OuterVisualizerComponent"));
+	OuterVisualizerComponent->SetHiddenInGame(true);
+	OuterVisualizerComponent->SetTranslucentSortPriority(1);
 #if WITH_EDITOR
 	// Do not show bounding box around cloner for better visibility
-	EndRadialComponent->SetIsVisualizationComponent(true);
+	OuterVisualizerComponent->SetIsVisualizationComponent(true);
 #endif
-	EndRadialComponent->bIsEditorOnly = false;
-	EndRadialComponent->SetupAttachment(SceneComponent);
+	OuterVisualizerComponent->bIsEditorOnly = false;
+	OuterVisualizerComponent->SetupAttachment(SceneComponent);
+
+	static const ConstructorHelpers::FObjectFinder<UMaterialInterface> VisualizerMaterial(TEXT("/Script/Engine.Material'/ClonerEffector/Materials/M_EffectorVisualizer.M_EffectorVisualizer'"));
+
+	if (VisualizerMaterial.Succeeded())
+	{
+		InnerVisualizerMaterial = UMaterialInstanceDynamic::Create(
+			VisualizerMaterial.Object,
+			this
+		);
+
+		InnerVisualizerMaterial->SetVectorParameterValue(VisualizerColorName, FLinearColor::Red);
+		InnerVisualizerMaterial->SetScalarParameterValue(VisualizerOpacityName, VisualizerOpacity);
+
+		OuterVisualizerMaterial = UMaterialInstanceDynamic::Create(
+			VisualizerMaterial.Object,
+			this
+		);
+
+		OuterVisualizerMaterial->SetVectorParameterValue(VisualizerColorName, FLinearColor::Blue);
+		OuterVisualizerMaterial->SetScalarParameterValue(VisualizerOpacityName, VisualizerOpacity);
+	}
 
 	if (!IsTemplate())
 	{
 		SceneComponent->TransformUpdated.AddUObject(this, &ACEEffectorActor::OnEffectorTransformed);
+		UCEEffectorSubsystem::OnSubsystemInitializedDelegate.AddUObject(this, &ACEEffectorActor::OnEffectorSubsystemInitialized);
 	}
 }
 
@@ -140,7 +86,7 @@ TCEPropertyChangeDispatcher<ACEEffectorActor> ACEEffectorActor::PropertyChangeDi
 	/** Effector */
 	{ GET_MEMBER_NAME_CHECKED(ACEEffectorActor, bEnabled), &ACEEffectorActor::OnEnabledChanged },
 	{ GET_MEMBER_NAME_CHECKED(ACEEffectorActor, Magnitude), &ACEEffectorActor::OnMagnitudeChanged },
-	{ GET_MEMBER_NAME_CHECKED(ACEEffectorActor, VisualizerThickness), &ACEEffectorActor::OnVisualizerThicknessChanged },
+	{ GET_MEMBER_NAME_CHECKED(ACEEffectorActor, VisualizerOpacity), &ACEEffectorActor::OnVisualizerOpacityChanged },
 	{ GET_MEMBER_NAME_CHECKED(ACEEffectorActor, bVisualizerSpriteVisible), &ACEEffectorActor::OnVisualizerSpriteVisibleChanged },
 	{ GET_MEMBER_NAME_CHECKED(ACEEffectorActor, Color), &ACEEffectorActor::OnColorChanged },
 	/** Type */
@@ -155,6 +101,9 @@ TCEPropertyChangeDispatcher<ACEEffectorActor> ACEEffectorActor::PropertyChangeDi
 	{ GET_MEMBER_NAME_CHECKED(ACEEffectorActor, RadialAngle), &ACEEffectorActor::OnRadialChanged },
 	{ GET_MEMBER_NAME_CHECKED(ACEEffectorActor, RadialMinRadius), &ACEEffectorActor::OnRadialChanged },
 	{ GET_MEMBER_NAME_CHECKED(ACEEffectorActor, RadialMaxRadius), &ACEEffectorActor::OnRadialChanged },
+	{ GET_MEMBER_NAME_CHECKED(ACEEffectorActor, TorusRadius), &ACEEffectorActor::OnTorusChanged },
+	{ GET_MEMBER_NAME_CHECKED(ACEEffectorActor, TorusInnerRadius), &ACEEffectorActor::OnTorusChanged },
+	{ GET_MEMBER_NAME_CHECKED(ACEEffectorActor, TorusOuterRadius), &ACEEffectorActor::OnTorusChanged },
 	/** Mode */
 	{ GET_MEMBER_NAME_CHECKED(ACEEffectorActor, Mode), &ACEEffectorActor::OnModeChanged },
 	{ GET_MEMBER_NAME_CHECKED(ACEEffectorActor, Offset), &ACEEffectorActor::OnTransformOptionsChanged },
@@ -228,6 +177,8 @@ void ACEEffectorActor::PostLoad()
 {
 	Super::PostLoad();
 
+	OnEffectorChanged();
+
 	if (!InternalCloners.IsEmpty())
 	{
 		for (const TWeakObjectPtr<ACEClonerActor>& ClonerWeak : InternalCloners)
@@ -240,8 +191,6 @@ void ACEEffectorActor::PostLoad()
 
 		InternalCloners.Empty();
 	}
-
-	OnEffectorChanged();
 }
 
 void ACEEffectorActor::PostDuplicate(EDuplicateMode::Type InDuplicateMode)
@@ -251,24 +200,32 @@ void ACEEffectorActor::PostDuplicate(EDuplicateMode::Type InDuplicateMode)
 	OnEffectorChanged();
 }
 
-int32 ACEEffectorActor::GetChannelIdentifier() const
+void ACEEffectorActor::RegisterToChannel()
 {
-	const int32 Identifier = ChannelData.GetIdentifier();
-
-	if (Identifier == INDEX_NONE)
+	if (IsValid(this)
+		&& ChannelData.GetIdentifier() == INDEX_NONE)
 	{
-		UCEEffectorSubsystem* EffectorSubsystem = !!GetWorld() ? UCEEffectorSubsystem::Get(GetWorld()) : nullptr;
-
 		// Register this effector to the effector channel
-		if (EffectorSubsystem)
+		if (UCEEffectorSubsystem* EffectorSubsystem = UCEEffectorSubsystem::Get(GetWorld()))
 		{
-			EffectorSubsystem->RegisterChannelEffector(const_cast<ACEEffectorActor*>(this));
+			EffectorSubsystem->RegisterChannelEffector(this);
 		}
 	}
+}
 
+int32 ACEEffectorActor::GetChannelIdentifier() const
+{
 	return ChannelData.GetIdentifier();
 }
 
+void ACEEffectorActor::OnEffectorSubsystemInitialized(const UWorld* InWorld)
+{
+	if (InWorld && GetWorld() == InWorld)
+	{
+		OnEffectorChanged();
+	}
+}
+
 FCEClonerEffectorChannelData& ACEEffectorActor::GetChannelData()
 {
 	return ChannelData;
@@ -442,6 +399,39 @@ void ACEEffectorActor::SetRadialMaxRadius(float InRadius)
 	OnRadialChanged();
 }
 
+void ACEEffectorActor::SetTorusRadius(float InRadius)
+{
+	if (FMath::IsNearlyEqual(InRadius, TorusRadius))
+	{
+		return;
+	}
+
+	TorusRadius = InRadius;
+	OnTorusChanged();
+}
+
+void ACEEffectorActor::SetTorusInnerRadius(float InRadius)
+{
+	if (FMath::IsNearlyEqual(InRadius, TorusInnerRadius))
+	{
+		return;
+	}
+
+	TorusInnerRadius = InRadius;
+	OnTorusChanged();
+}
+
+void ACEEffectorActor::SetTorusOuterRadius(float InRadius)
+{
+	if (FMath::IsNearlyEqual(InRadius, TorusOuterRadius))
+	{
+		return;
+	}
+
+	TorusOuterRadius = InRadius;
+	OnTorusChanged();
+}
+
 void ACEEffectorActor::SetOffset(const FVector& InOffset)
 {
 	if (InOffset.Equals(Offset))
@@ -503,20 +493,20 @@ void ACEEffectorActor::SetEnabled(bool bInEnable)
 	OnEnabledChanged();
 }
 
-void ACEEffectorActor::SetVisualizerThickness(float InThickness)
+void ACEEffectorActor::SetVisualizerOpacity(float InOpacity)
 {
-	if (FMath::IsNearlyEqual(VisualizerThickness, InThickness))
+	if (FMath::IsNearlyEqual(VisualizerOpacity, InOpacity))
 	{
 		return;
 	}
 
-	if (VisualizerThickness < 0.1f || VisualizerThickness > 10.f)
+	if (VisualizerOpacity < 0.f || VisualizerOpacity > 1.f)
 	{
 		return;
 	}
 
-	VisualizerThickness = InThickness;
-	OnVisualizerThicknessChanged();
+	VisualizerOpacity = InOpacity;
+	OnVisualizerOpacityChanged();
 }
 
 void ACEEffectorActor::SetMode(ECEClonerEffectorMode InMode)
@@ -541,6 +531,11 @@ void ACEEffectorActor::SetTargetActor(AActor* InTargetActor)
 	OnTargetActorChanged();
 }
 
+void ACEEffectorActor::SetTargetActorWeak(const TWeakObjectPtr<AActor>& InTargetActor)
+{
+	SetTargetActor(InTargetActor.Get());
+}
+
 void ACEEffectorActor::SetAttractionForceStrength(float InForceStrength)
 {
 	if (FMath::IsNearlyEqual(AttractionForceStrength, InForceStrength))
@@ -795,10 +790,7 @@ void ACEEffectorActor::OnEffectorTransformed(USceneComponent* InUpdatedComponent
 	OnTransformChanged();
 
 	// Update when scaled or rotated
-	OnSphereChanged();
-	OnBoxChanged();
-	OnPlaneChanged();
-	OnRadialChanged();
+	UpdateEffectorTypes();
 
 	// update if self
 	const AActor* InternalTargetActor = InternalTargetActorWeak.Get();
@@ -820,7 +812,7 @@ void ACEEffectorActor::OnEnabledChanged()
 		// Hide visualization components
 		OnTypeChanged();
 		// Editor
-		OnVisualizerThicknessChanged();
+		OnVisualizerOpacityChanged();
 		OnVisualizerSpriteVisibleChanged();
 	}
 }
@@ -841,45 +833,16 @@ void ACEEffectorActor::OnModeChanged()
 
 void ACEEffectorActor::OnTypeChanged()
 {
-	// Sphere
-	if (InnerSphereComponent && OuterSphereComponent)
-	{
-		InnerSphereComponent->SetVisibility(Type == ECEClonerEffectorType::Sphere);
-		OuterSphereComponent->SetVisibility(Type == ECEClonerEffectorType::Sphere);
-	}
-
-	// Box
-	if (InnerBoxComponent && OuterBoxComponent)
-	{
-		InnerBoxComponent->SetVisibility(Type == ECEClonerEffectorType::Box);
-		OuterBoxComponent->SetVisibility(Type == ECEClonerEffectorType::Box);
-	}
-
-	// Plane
-	if (InnerPlaneComponent && OuterPlaneComponent)
-	{
-		InnerPlaneComponent->SetVisibility(Type == ECEClonerEffectorType::Plane);
-		OuterPlaneComponent->SetVisibility(Type == ECEClonerEffectorType::Plane);
-	}
-
-	// Radial
-	if (BeginRadialComponent && EndRadialComponent)
-	{
-		BeginRadialComponent->SetVisibility(Type == ECEClonerEffectorType::Radial);
-		EndRadialComponent->SetVisibility(Type == ECEClonerEffectorType::Radial);
-	}
-
 	ChannelData.Type = Type;
 
 	// Update type data
-	OnSphereChanged();
-	OnBoxChanged();
-	OnPlaneChanged();
-	OnRadialChanged();
+	UpdateEffectorTypes();
 }
 
 void ACEEffectorActor::OnEffectorChanged()
 {
+	RegisterToChannel();
+
 	if (!bEnabled)
 	{
 		OnEnabledChanged();
@@ -894,7 +857,7 @@ void ACEEffectorActor::OnEffectorChanged()
 		OnForceOptionsChanged();
 		OnColorChanged();
 		// Editor
-		OnVisualizerThicknessChanged();
+		OnVisualizerOpacityChanged();
 		OnVisualizerSpriteVisibleChanged();
 	}
 }
@@ -914,15 +877,8 @@ void ACEEffectorActor::OnTransformChanged()
 void ACEEffectorActor::OnBoxChanged()
 {
 	InnerExtent = ClampVector(InnerExtent, FVector::ZeroVector, OuterExtent);
-	// Max()
 	OuterExtent = OuterExtent.ComponentMax(InnerExtent);
 
-	if (InnerBoxComponent && OuterBoxComponent)
-	{
-		InnerBoxComponent->SetBoxExtent(InnerExtent);
-		OuterBoxComponent->SetBoxExtent(OuterExtent);
-	}
-
 	if (Type != ECEClonerEffectorType::Box)
 	{
 		return;
@@ -930,6 +886,38 @@ void ACEEffectorActor::OnBoxChanged()
 
 	ChannelData.InnerExtent = GetInnerExtent();
 	ChannelData.OuterExtent = GetOuterExtent();
+
+	// Update visualizer
+	static const FName BoxInnerExtentName = GET_MEMBER_NAME_CHECKED(ACEEffectorActor, InnerExtent);
+	static const FName BoxOuterExtentName = GET_MEMBER_NAME_CHECKED(ACEEffectorActor, OuterExtent);
+	TValueOrError<FVector*, EPropertyBagResult> BoxInnerExtentValue = VisualizerData.GetValueStruct<FVector>(BoxInnerExtentName);
+	TValueOrError<FVector*, EPropertyBagResult> BoxOuterExtentValue = VisualizerData.GetValueStruct<FVector>(BoxOuterExtentName);
+
+	if (!BoxInnerExtentValue.HasValue()
+		|| *BoxInnerExtentValue.GetValue() != InnerExtent)
+	{
+		UpdateVisualizer(InnerVisualizerId, [this](UDynamicMesh* InMesh)
+		{
+			constexpr FGeometryScriptPrimitiveOptions PrimitiveOptions;
+			UGeometryScriptLibrary_MeshPrimitiveFunctions::AppendBox(InMesh, PrimitiveOptions, FTransform(FVector(0, 0, -InnerExtent.Z)), InnerExtent.X * 2, InnerExtent.Y * 2, InnerExtent.Z * 2);
+		});
+	}
+
+	if (!BoxOuterExtentValue.HasValue()
+		|| *BoxOuterExtentValue.GetValue() != OuterExtent)
+	{
+		UpdateVisualizer(OuterVisualizerId, [this](UDynamicMesh* InMesh)
+		{
+			constexpr FGeometryScriptPrimitiveOptions PrimitiveOptions;
+			UGeometryScriptLibrary_MeshPrimitiveFunctions::AppendBox(InMesh, PrimitiveOptions, FTransform(FVector(0, 0, -OuterExtent.Z)), OuterExtent.X * 2, OuterExtent.Y * 2, OuterExtent.Z * 2);
+		});
+	}
+
+	VisualizerData.Reset();
+	VisualizerData.AddProperty(BoxInnerExtentName, EPropertyBagPropertyType::Struct);
+	VisualizerData.AddProperty(BoxOuterExtentName, EPropertyBagPropertyType::Struct);
+	VisualizerData.SetValueStruct(BoxInnerExtentName, InnerExtent);
+	VisualizerData.SetValueStruct(BoxOuterExtentName, OuterExtent);
 }
 
 void ACEEffectorActor::OnSphereChanged()
@@ -937,12 +925,6 @@ void ACEEffectorActor::OnSphereChanged()
 	InnerRadius = FMath::Clamp(InnerRadius, 0.f, OuterRadius);
 	OuterRadius = FMath::Max(OuterRadius, InnerRadius);
 
-	if (InnerSphereComponent && OuterSphereComponent)
-	{
-		InnerSphereComponent->SetSphereRadius(InnerRadius);
-		OuterSphereComponent->SetSphereRadius(OuterRadius);
-	}
-
 	if (Type != ECEClonerEffectorType::Sphere)
 	{
 		return;
@@ -950,43 +932,81 @@ void ACEEffectorActor::OnSphereChanged()
 
 	ChannelData.InnerExtent = FVector(GetInnerRadius());
 	ChannelData.OuterExtent = FVector(GetOuterRadius());
-}
 
-void ACEEffectorActor::OnPlaneChanged()
-{
-	static const FVector PlaneAxis = FVector::LeftVector;
-	const FVector InnerPlane = PlaneAxis * FVector(-PlaneSpacing/2);
-	const FVector OuterPlane = PlaneAxis * FVector(PlaneSpacing/2);
+	// Update visualizer
+	static const FName SphereInnerRadiusName = GET_MEMBER_NAME_CHECKED(ACEEffectorActor, InnerRadius);
+	static const FName SphereOuterRadiusName = GET_MEMBER_NAME_CHECKED(ACEEffectorActor, OuterRadius);
+	TValueOrError<float, EPropertyBagResult> SphereInnerRadiusValue = VisualizerData.GetValueFloat(SphereInnerRadiusName);
+	TValueOrError<float, EPropertyBagResult> SphereOuterRadiusValue = VisualizerData.GetValueFloat(SphereOuterRadiusName);
 
-	if (InnerPlaneComponent && OuterPlaneComponent)
+	if (!SphereInnerRadiusValue.HasValue()
+		|| SphereInnerRadiusValue.GetValue() != InnerRadius)
 	{
-		static const FVector InnerSize(100, 0, 100);
-		InnerPlaneComponent->SetRelativeLocation(InnerPlane);
-		InnerPlaneComponent->SetBoxExtent(InnerSize);
+		UpdateVisualizer(InnerVisualizerId, [this](UDynamicMesh* InMesh)
+		{
+			constexpr FGeometryScriptPrimitiveOptions PrimitiveOptions;
+			UGeometryScriptLibrary_MeshPrimitiveFunctions::AppendSphereLatLong(InMesh, PrimitiveOptions, FTransform::Identity, InnerRadius);
+		});
+	}
 
-		static const FVector OuterSize(200, 0, 200);
-		OuterPlaneComponent->SetRelativeLocation(OuterPlane);
-		OuterPlaneComponent->SetBoxExtent(OuterSize);
+	if (!SphereOuterRadiusValue.HasValue()
+		|| SphereOuterRadiusValue.GetValue() != OuterRadius)
+	{
+		UpdateVisualizer(OuterVisualizerId, [this](UDynamicMesh* InMesh)
+		{
+			constexpr FGeometryScriptPrimitiveOptions PrimitiveOptions;
+			UGeometryScriptLibrary_MeshPrimitiveFunctions::AppendSphereLatLong(InMesh, PrimitiveOptions, FTransform::Identity, OuterRadius);
+		});
 	}
 
+	VisualizerData.Reset();
+	VisualizerData.AddProperty(SphereInnerRadiusName, EPropertyBagPropertyType::Float);
+	VisualizerData.AddProperty(SphereOuterRadiusName, EPropertyBagPropertyType::Float);
+	VisualizerData.SetValueFloat(SphereInnerRadiusName, InnerRadius);
+	VisualizerData.SetValueFloat(SphereOuterRadiusName, OuterRadius);
+}
+
+void ACEEffectorActor::OnPlaneChanged()
+{
 	if (Type != ECEClonerEffectorType::Plane)
 	{
 		return;
 	}
 
-	ChannelData.InnerExtent = PlaneAxis;
+	ChannelData.InnerExtent = FVector::LeftVector;
 	ChannelData.OuterExtent = FVector(PlaneSpacing);
+
+	// Update visualizer
+	static const FName PlaneSpacingName = GET_MEMBER_NAME_CHECKED(ACEEffectorActor, PlaneSpacing);
+	TValueOrError<float, EPropertyBagResult> PlaneSpacingValue = VisualizerData.GetValueFloat(PlaneSpacingName);
+
+	if (!PlaneSpacingValue.HasValue()
+		|| PlaneSpacingValue.GetValue() != PlaneSpacing)
+	{
+		UpdateVisualizer(InnerVisualizerId, [this](UDynamicMesh* InMesh)
+		{
+			const FVector InnerPlane = FVector::LeftVector * FVector(-PlaneSpacing/2);
+			constexpr FGeometryScriptPrimitiveOptions PrimitiveOptions;
+			UGeometryScriptLibrary_MeshPrimitiveFunctions::AppendRectangleXY(InMesh, PrimitiveOptions, FTransform(FRotator(0, 0, 90), InnerPlane), 250, 250);
+		});
+
+		UpdateVisualizer(OuterVisualizerId, [this](UDynamicMesh* InMesh)
+		{
+			const FVector OuterPlane = FVector::LeftVector * FVector(PlaneSpacing/2);
+			constexpr FGeometryScriptPrimitiveOptions PrimitiveOptions;
+			UGeometryScriptLibrary_MeshPrimitiveFunctions::AppendRectangleXY(InMesh, PrimitiveOptions, FTransform(FRotator(0, 0, 90), OuterPlane), 500, 500);
+		});
+	}
+
+	VisualizerData.Reset();
+	VisualizerData.AddProperty(PlaneSpacingName, EPropertyBagPropertyType::Float);
+	VisualizerData.SetValueFloat(PlaneSpacingName, PlaneSpacing);
 }
 
 void ACEEffectorActor::OnRadialChanged()
 {
 	RadialMinRadius = FMath::Min(RadialMinRadius, RadialMaxRadius);
 	RadialMaxRadius = FMath::Max(RadialMinRadius, RadialMaxRadius);
-	
-	if (BeginRadialComponent && EndRadialComponent)
-	{
-		EndRadialComponent->SetRelativeRotation(FRotator(0, -90 + RadialAngle, 0));
-	}
 
 	if (Type != ECEClonerEffectorType::Radial)
 	{
@@ -995,6 +1015,111 @@ void ACEEffectorActor::OnRadialChanged()
 
 	ChannelData.InnerExtent = FVector::LeftVector;
 	ChannelData.OuterExtent = FVector(RadialAngle, RadialMinRadius, RadialMaxRadius);
+
+	// Update visualizer
+	static const FName RadialAngleName = GET_MEMBER_NAME_CHECKED(ACEEffectorActor, RadialAngle);
+	static const FName RadialMinRadiusName = GET_MEMBER_NAME_CHECKED(ACEEffectorActor, RadialMinRadius);
+	static const FName RadialMaxRadiusName = GET_MEMBER_NAME_CHECKED(ACEEffectorActor, RadialMaxRadius);
+	TValueOrError<float, EPropertyBagResult> RadialAngleValue = VisualizerData.GetValueFloat(RadialAngleName);
+	TValueOrError<float, EPropertyBagResult> RadialMinRadiusValue = VisualizerData.GetValueFloat(RadialMinRadiusName);
+	TValueOrError<float, EPropertyBagResult> RadialMaxRadiusValue = VisualizerData.GetValueFloat(RadialMaxRadiusName);
+
+	if (!RadialMaxRadiusValue.HasValue()
+		|| RadialMaxRadiusValue.GetValue() != RadialMaxRadius
+		|| !RadialMinRadiusValue.HasValue()
+		|| RadialMinRadiusValue.GetValue() != RadialMinRadius
+		|| !RadialAngleValue.HasValue()
+		|| RadialAngleValue.GetValue() != RadialAngle)
+	{
+		UpdateVisualizer(InnerVisualizerId, [this](UDynamicMesh* InMesh)
+		{
+			constexpr FGeometryScriptPrimitiveOptions PrimitiveOptions;
+			UGeometryScriptLibrary_MeshPrimitiveFunctions::AppendDisc(InMesh, PrimitiveOptions, FTransform(FRotator(0, -90, 0)), RadialMaxRadius, 16, 0, 0, RadialAngle / 2, RadialMinRadius);
+		});
+
+		UpdateVisualizer(OuterVisualizerId, [this](UDynamicMesh* InMesh)
+		{
+			constexpr FGeometryScriptPrimitiveOptions PrimitiveOptions;
+			UGeometryScriptLibrary_MeshPrimitiveFunctions::AppendDisc(InMesh, PrimitiveOptions, FTransform(FRotator(0, -90, 0)), RadialMaxRadius, 16, 0, RadialAngle / 2, RadialAngle, RadialMinRadius);
+		});
+	}
+
+	VisualizerData.Reset();
+	VisualizerData.AddProperty(RadialAngleName, EPropertyBagPropertyType::Float);
+	VisualizerData.AddProperty(RadialMinRadiusName, EPropertyBagPropertyType::Float);
+	VisualizerData.AddProperty(RadialMaxRadiusName, EPropertyBagPropertyType::Float);
+	VisualizerData.SetValueFloat(RadialAngleName, RadialAngle);
+	VisualizerData.SetValueFloat(RadialMinRadiusName, RadialMinRadius);
+	VisualizerData.SetValueFloat(RadialMaxRadiusName, RadialMaxRadius);
+}
+
+void ACEEffectorActor::OnTorusChanged()
+{
+	TorusOuterRadius = FMath::Min(TorusRadius, TorusOuterRadius);
+	TorusInnerRadius = FMath::Min(TorusInnerRadius, TorusOuterRadius);
+	TorusOuterRadius = FMath::Max(TorusInnerRadius, TorusOuterRadius);
+
+	if (Type != ECEClonerEffectorType::Torus)
+	{
+		return;
+	}
+
+	ChannelData.InnerExtent = FVector::ZAxisVector;
+	ChannelData.OuterExtent = FVector(TorusInnerRadius, TorusOuterRadius, TorusRadius);
+
+	// Update visualizer
+	static const FName TorusRadiusName = GET_MEMBER_NAME_CHECKED(ACEEffectorActor, TorusRadius);
+	static const FName TorusInnerRadiusName = GET_MEMBER_NAME_CHECKED(ACEEffectorActor, TorusInnerRadius);
+	static const FName TorusOuterRadiusName = GET_MEMBER_NAME_CHECKED(ACEEffectorActor, TorusOuterRadius);
+	TValueOrError<float, EPropertyBagResult> TorusRadiusValue = VisualizerData.GetValueFloat(TorusRadiusName);
+	TValueOrError<float, EPropertyBagResult> TorusInnerRadiusValue = VisualizerData.GetValueFloat(TorusInnerRadiusName);
+	TValueOrError<float, EPropertyBagResult> TorusOuterRadiusValue = VisualizerData.GetValueFloat(TorusOuterRadiusName);
+
+	if (!TorusRadiusValue.HasValue()
+		|| TorusRadiusValue.GetValue() != TorusRadius
+		|| !TorusInnerRadiusValue.HasValue()
+		|| TorusInnerRadiusValue.GetValue() != TorusInnerRadius)
+	{
+		UpdateVisualizer(InnerVisualizerId, [this](UDynamicMesh* InMesh)
+		{
+			constexpr FGeometryScriptPrimitiveOptions PrimitiveOptions;
+			constexpr FGeometryScriptRevolveOptions RevolveOptions;
+			UGeometryScriptLibrary_MeshPrimitiveFunctions::AppendTorus(InMesh, PrimitiveOptions, FTransform(FVector(0, 0, -TorusInnerRadius)), RevolveOptions, TorusRadius, TorusInnerRadius);
+		});
+	}
+
+	if (!TorusRadiusValue.HasValue()
+		|| TorusRadiusValue.GetValue() != TorusRadius
+		|| !TorusOuterRadiusValue.HasValue()
+		|| TorusOuterRadiusValue.GetValue() != TorusOuterRadius)
+	{
+		UpdateVisualizer(OuterVisualizerId, [this](UDynamicMesh* InMesh)
+		{
+			constexpr FGeometryScriptPrimitiveOptions PrimitiveOptions;
+			constexpr FGeometryScriptRevolveOptions RevolveOptions;
+			UGeometryScriptLibrary_MeshPrimitiveFunctions::AppendTorus(InMesh, PrimitiveOptions, FTransform(FVector(0, 0, -TorusOuterRadius)), RevolveOptions, TorusRadius, TorusOuterRadius);
+		});
+	}
+
+	VisualizerData.Reset();
+	VisualizerData.AddProperty(TorusRadiusName, EPropertyBagPropertyType::Float);
+	VisualizerData.AddProperty(TorusInnerRadiusName, EPropertyBagPropertyType::Float);
+	VisualizerData.AddProperty(TorusOuterRadiusName, EPropertyBagPropertyType::Float);
+	VisualizerData.SetValueFloat(TorusRadiusName, TorusRadius);
+	VisualizerData.SetValueFloat(TorusInnerRadiusName, TorusInnerRadius);
+	VisualizerData.SetValueFloat(TorusOuterRadiusName, TorusOuterRadius);
+}
+
+void ACEEffectorActor::OnUnboundChanged()
+{
+	if (Type != ECEClonerEffectorType::Unbound)
+	{
+		return;
+	}
+
+	// Clear visualizers
+	UpdateVisualizer(InnerVisualizerId, [](UDynamicMesh* InMesh){});
+	UpdateVisualizer(OuterVisualizerId, [](UDynamicMesh* InMesh){});
 }
 
 void ACEEffectorActor::OnMagnitudeChanged()
@@ -1071,6 +1196,29 @@ void ACEEffectorActor::OnTargetActorDestroyed(AActor* InActor)
 	}
 }
 
+void ACEEffectorActor::UpdateVisualizer(int32 InVisualizerId, TFunction<void(UDynamicMesh*)> InMeshFunction) const
+{
+	UDynamicMeshComponent* MeshComponent = InVisualizerId == InnerVisualizerId ? InnerVisualizerComponent : OuterVisualizerComponent;
+
+	if (!MeshComponent)
+	{
+		return;
+	}
+
+	UDynamicMesh* DynamicMesh = MeshComponent->GetDynamicMesh();
+
+	DynamicMesh->EditMesh([](FDynamicMesh3& InMesh)
+	{
+		InMesh.Clear();
+	});
+
+	InMeshFunction(DynamicMesh);
+
+	// Apply material
+	UMaterialInstanceDynamic* VisualizerMaterial = InVisualizerId == InnerVisualizerId ? InnerVisualizerMaterial : OuterVisualizerMaterial;
+	MeshComponent->SetMaterial(0, VisualizerMaterial);
+}
+
 void ACEEffectorActor::OnNoiseFieldOptionsChanged()
 {
 	if (Mode != ECEClonerEffectorMode::NoiseField)
@@ -1108,33 +1256,18 @@ void ACEEffectorActor::OnTargetOptionsChanged()
 	ChannelData.ScaleDelta = FVector::OneVector;
 }
 
-void ACEEffectorActor::OnVisualizerThicknessChanged()
+void ACEEffectorActor::OnVisualizerOpacityChanged()
 {
-	VisualizerThickness = FMath::Clamp(VisualizerThickness, 0.1f, 10.f);
-
-	if (InnerSphereComponent && OuterSphereComponent)
-	{
-		InnerSphereComponent->SetLineThickness(VisualizerThickness);
-		OuterSphereComponent->SetLineThickness(VisualizerThickness);
-	}
-
-	if (InnerBoxComponent && OuterBoxComponent)
-	{
-		InnerBoxComponent->SetLineThickness(VisualizerThickness);
-		OuterBoxComponent->SetLineThickness(VisualizerThickness);
-	}
+	VisualizerOpacity = FMath::Clamp(VisualizerOpacity, 0.f, 1.f);
 
-	if (InnerPlaneComponent && OuterPlaneComponent)
+	if (InnerVisualizerMaterial)
 	{
-		InnerPlaneComponent->SetLineThickness(VisualizerThickness);
-		OuterPlaneComponent->SetLineThickness(VisualizerThickness);
+		InnerVisualizerMaterial->SetScalarParameterValue(VisualizerOpacityName, VisualizerOpacity);
 	}
 
-	if (BeginRadialComponent && EndRadialComponent)
+	if (OuterVisualizerMaterial)
 	{
-		const float ArrowScale = VisualizerThickness / 2;
-		BeginRadialComponent->SetRelativeScale3D(FVector(10, ArrowScale, ArrowScale));
-		EndRadialComponent->SetRelativeScale3D(FVector(10, ArrowScale, ArrowScale));
+		OuterVisualizerMaterial->SetScalarParameterValue(VisualizerOpacityName, VisualizerOpacity);
 	}
 }
 
@@ -1151,6 +1284,16 @@ void ACEEffectorActor::OnColorChanged()
 	ChannelData.Color = Color;
 }
 
+void ACEEffectorActor::UpdateEffectorTypes()
+{
+	OnSphereChanged();
+	OnBoxChanged();
+	OnPlaneChanged();
+	OnRadialChanged();
+	OnTorusChanged();
+	OnUnboundChanged();
+}
+
 void ACEEffectorActor::OnForceOptionsChanged()
 {
 	if (bOrientationForceEnabled)
diff --git a/Engine/Plugins/Experimental/ClonerEffector/Source/ClonerEffector/Private/Subsystems/CEEffectorSubsystem.cpp b/Engine/Plugins/Experimental/ClonerEffector/Source/ClonerEffector/Private/Subsystems/CEEffectorSubsystem.cpp
--- a/Engine/Plugins/Experimental/ClonerEffector/Source/ClonerEffector/Private/Subsystems/CEEffectorSubsystem.cpp
+++ b/Engine/Plugins/Experimental/ClonerEffector/Source/ClonerEffector/Private/Subsystems/CEEffectorSubsystem.cpp
@@ -49,9 +49,9 @@ bool UCEEffectorSubsystem::RegisterChannelEffector(ACEEffectorActor* InEffector)
 		return false;
 	}
 
-	int32 EffectorIndex = INDEX_NONE;
+	int32 EffectorIndex = EffectorsWeak.Find(InEffector);
 
-	if (!EffectorsWeak.Contains(InEffector))
+	if (EffectorIndex == INDEX_NONE)
 	{
 		EffectorIndex = EffectorsWeak.Add(InEffector);
 		UE_LOG(LogACEEffectorSubsystem, Log, TEXT("%s effector registered in channel %i"), *InEffector->GetActorNameOrLabel(), EffectorIndex);
diff --git a/Engine/Plugins/Experimental/ClonerEffector/Source/ClonerEffector/Public/CEClonerEffectorShared.h b/Engine/Plugins/Experimental/ClonerEffector/Source/ClonerEffector/Public/CEClonerEffectorShared.h
--- a/Engine/Plugins/Experimental/ClonerEffector/Source/ClonerEffector/Public/CEClonerEffectorShared.h
+++ b/Engine/Plugins/Experimental/ClonerEffector/Source/ClonerEffector/Public/CEClonerEffectorShared.h
@@ -154,7 +154,9 @@ enum class ECEClonerEffectorType : uint8
 	/** All clones will be affected by the effector with the same max weight */
 	Unbound,
 	/** All clones within the angle range will be affected */
-	Radial
+	Radial,
+	/** All clones inside the torus radius will be affected by the effector */
+	Torus
 };
 
 /** Enumerates the effector mode available */
@@ -313,8 +315,11 @@ struct FCEClonerAttachmentTree
 	}
 };
 
+USTRUCT()
 struct FCEClonerEffectorChannelData
 {
+	GENERATED_BODY()
+
 	friend class UCEEffectorSubsystem;
 
 	/** General */
@@ -381,7 +386,8 @@ struct FCEClonerEffectorChannelData
 	FVector GravityForceAcceleration = FVector::ZeroVector;
 
 protected:
-	/** Cache effector identifier to detect a change and update cloners DI */
+	/** Cached effector identifier to detect a change and update cloners DI */
+	UPROPERTY(EditInstanceOnly, Category="Effector")
 	int32 Identifier = INDEX_NONE;
 
 	void Write(UNiagaraDataChannelWriter* InWriter) const;
diff --git a/Engine/Plugins/Experimental/ClonerEffector/Source/ClonerEffector/Public/Effector/CEEffectorActor.h b/Engine/Plugins/Experimental/ClonerEffector/Source/ClonerEffector/Public/Effector/CEEffectorActor.h
--- a/Engine/Plugins/Experimental/ClonerEffector/Source/ClonerEffector/Public/Effector/CEEffectorActor.h
+++ b/Engine/Plugins/Experimental/ClonerEffector/Source/ClonerEffector/Public/Effector/CEEffectorActor.h
@@ -5,13 +5,11 @@
 #include "CEClonerEffectorShared.h"
 #include "CEPropertyChangeDispatcher.h"
 #include "GameFramework/Actor.h"
+#include "PropertyBag.h"
 #include "CEEffectorActor.generated.h"
 
-class UNiagaraDataChannelWriter;
 class ACEClonerActor;
-class UArrowComponent;
-class USphereComponent;
-class UBoxComponent;
+class UDynamicMeshComponent;
 
 UCLASS(MinimalAPI, BlueprintType, HideCategories=(Rendering,Replication,Collision,HLOD,Physics,Networking,Input,Actor,Cooking,LevelInstance), DisplayName = "Motion Design Effector Actor")
 class ACEEffectorActor : public AActor
@@ -136,6 +134,33 @@ public:
 		return RadialMaxRadius;
 	}
 
+	UFUNCTION(BlueprintCallable, Category="Effector")
+	void SetTorusRadius(float InRadius);
+
+	UFUNCTION(BlueprintPure, Category="Effector")
+	float GetTorusRadius() const
+	{
+		return TorusRadius;
+	}
+
+	UFUNCTION(BlueprintCallable, Category="Effector")
+	void SetTorusInnerRadius(float InRadius);
+
+	UFUNCTION(BlueprintPure, Category="Effector")
+	float GetTorusInnerRadius() const
+	{
+		return TorusInnerRadius;
+	}
+
+	UFUNCTION(BlueprintCallable, Category="Effector")
+	void SetTorusOuterRadius(float InRadius);
+
+	UFUNCTION(BlueprintPure, Category="Effector")
+	float GetTorusOuterRadius() const
+	{
+		return TorusOuterRadius;
+	}
+
 	UFUNCTION(BlueprintCallable, Category="Effector")
 	CLONEREFFECTOR_API void SetOffset(const FVector& InOffset);
 
@@ -190,6 +215,13 @@ public:
 		return TargetActorWeak.Get();
 	}
 
+	void SetTargetActorWeak(const TWeakObjectPtr<AActor>& InTargetActor);
+
+	TWeakObjectPtr<AActor> GetTargetActorWeak() const
+	{
+		return TargetActorWeak;
+	}
+
 	UFUNCTION(BlueprintCallable, Category="Effector")
 	CLONEREFFECTOR_API void SetOrientationForceEnabled(bool bInForceEnabled);
 
@@ -380,12 +412,12 @@ public:
 	}
 
 	UFUNCTION(BlueprintCallable, Category="Effector")
-	CLONEREFFECTOR_API void SetVisualizerThickness(float InThickness);
+	CLONEREFFECTOR_API void SetVisualizerOpacity(float InOpacity);
 
 	UFUNCTION(BlueprintPure, Category="Effector")
-	float GetVisualizerThickness() const
+	float GetVisualizerOpacity() const
 	{
-		return VisualizerThickness;
+		return VisualizerOpacity;
 	}
 
 #if WITH_EDITOR
@@ -404,6 +436,13 @@ protected:
 	DECLARE_MULTICAST_DELEGATE_OneParam(FOnEffectorIdentifierChanged, ACEEffectorActor* /** InEffectorActor */)
 	static FOnEffectorIdentifierChanged OnEffectorRefreshClonerDelegate;
 
+	static constexpr int32 InnerVisualizerId = 0;
+	static constexpr int32 OuterVisualizerId = 1;
+	static constexpr TCHAR VisualizerColorName[] = TEXT("VisualizerColor");
+	static constexpr TCHAR VisualizerOpacityName[] = TEXT("VisualizerOpacity");
+
+	void UpdateVisualizer(int32 InVisualizerId, TFunction<void(UDynamicMesh*)> InMeshFunction) const;
+
 	//~ Begin AActor
 	virtual void Destroyed() override;
 	virtual void PostActorCreated() override;
@@ -422,7 +461,9 @@ protected:
 	virtual void PostDuplicate(EDuplicateMode::Type InDuplicateMode) override;
 	//~ End UObject
 
+	void RegisterToChannel();
 	int32 GetChannelIdentifier() const;
+	void OnEffectorSubsystemInitialized(const UWorld* InWorld);
 
 	FCEClonerEffectorChannelData& GetChannelData();
 
@@ -455,6 +496,12 @@ protected:
 	/** Update values for radial type effector */
 	void OnRadialChanged();
 
+	/** Update values for torus type effector */
+	void OnTorusChanged();
+
+	/** Update values for unbound type effector */
+	void OnUnboundChanged();
+
 	/** Called when transform of effector has changed */
 	void OnTransformChanged();
 
@@ -482,8 +529,8 @@ protected:
 	/** Called when force enabled state changed */
 	void OnForceEnabledChanged();
 
-	/** Update thickness of components of this effector */
-	void OnVisualizerThicknessChanged();
+	/** Update opacity of components of this effector */
+	void OnVisualizerOpacityChanged();
 
 	/** Update sprite visibility of this effector */
 	void OnVisualizerSpriteVisibleChanged();
@@ -491,12 +538,15 @@ protected:
 	/** Update particle color affected by this effector */
 	void OnColorChanged();
 
+	/** Update all types options */
+	void UpdateEffectorTypes();
+
 	/** Is this effector enabled/disabled on linked cloners */
 	UPROPERTY(EditAnywhere, BlueprintReadWrite, Setter="SetEnabled", Getter="GetEnabled", Category="Effector")
 	bool bEnabled = true;
 
 	/** The ratio effect of the effector on clones */
-	UPROPERTY(EditAnywhere, BlueprintReadWrite, Setter="SetMagnitude", Getter="GetMagnitude", Category="Effector", meta=(ClampMin="0", ClampMax="1"))
+	UPROPERTY(EditAnywhere, BlueprintReadWrite, Setter, Getter, Category="Effector", meta=(ClampMin="0", ClampMax="1"))
 	float Magnitude = 1.f;
 
 	/** Affected clones color passed over to material */
@@ -504,7 +554,7 @@ protected:
 	FLinearColor Color = FLinearColor::Red;
 
 	/** Type of effector to apply on cloners instances */
-	UPROPERTY(EditAnywhere, BlueprintReadWrite, Setter="SetType", Getter="GetType", Category="Type")
+	UPROPERTY(EditAnywhere, BlueprintReadWrite, Setter, Getter, Category="Type")
 	ECEClonerEffectorType Type = ECEClonerEffectorType::Sphere;
 
 	/** Invert the type effect, instead of affecting the inside of a zone, will affect the outside */
@@ -512,27 +562,27 @@ protected:
 	bool bInvertType = false;
 
 	/** Weight easing function applied to lerp transforms */
-	UPROPERTY(EditAnywhere, BlueprintReadWrite, Setter="SetEasing", Getter="GetEasing", Category="Type", meta=(EditCondition="Type != ECEClonerEffectorType::Unbound", EditConditionHides))
+	UPROPERTY(EditAnywhere, BlueprintReadWrite, Setter, Getter, Category="Type", meta=(EditCondition="Type != ECEClonerEffectorType::Unbound", EditConditionHides))
 	ECEClonerEasing Easing = ECEClonerEasing::Linear;
 
 	/** Inner radius of sphere, all clones inside will be affected with a maximum weight */
-	UPROPERTY(EditAnywhere, BlueprintReadWrite, Setter="SetInnerRadius", Getter="GetInnerRadius", Category="Type", meta=(ClampMin="0", EditCondition="Type == ECEClonerEffectorType::Sphere", EditConditionHides))
+	UPROPERTY(EditAnywhere, BlueprintReadWrite, Setter, Getter, Category="Type", meta=(ClampMin="0", EditCondition="Type == ECEClonerEffectorType::Sphere", EditConditionHides))
 	float InnerRadius = 50.f;
 
 	/** Outer radius of sphere, all clones outside will not be affected */
-	UPROPERTY(EditAnywhere, BlueprintReadWrite, Setter="SetOuterRadius", Getter="GetOuterRadius", Category="Type", meta=(ClampMin="0", EditCondition="Type == ECEClonerEffectorType::Sphere", EditConditionHides))
+	UPROPERTY(EditAnywhere, BlueprintReadWrite, Setter, Getter, Category="Type", meta=(ClampMin="0", EditCondition="Type == ECEClonerEffectorType::Sphere", EditConditionHides))
 	float OuterRadius = 200.f;
 
 	/** Inner extent of box, all clones inside will be affected with a maximum weight */
-	UPROPERTY(EditAnywhere, BlueprintReadWrite, Setter="SetInnerExtent", Getter="GetInnerExtent", Category="Type", meta=(ClampMin="0", EditCondition="Type == ECEClonerEffectorType::Box", EditConditionHides))
+	UPROPERTY(EditAnywhere, BlueprintReadWrite, Setter, Getter, Category="Type", meta=(ClampMin="0", AllowPreserveRatio, EditCondition="Type == ECEClonerEffectorType::Box", EditConditionHides))
 	FVector InnerExtent = FVector(50.f);
 
 	/** Outer extent of box, all clones outside will not be affected */
-	UPROPERTY(EditAnywhere, BlueprintReadWrite, Setter="SetOuterExtent", Getter="GetOuterExtent", Category="Type", meta=(ClampMin="0", EditCondition="Type == ECEClonerEffectorType::Box", EditConditionHides))
+	UPROPERTY(EditAnywhere, BlueprintReadWrite, Setter, Getter, Category="Type", meta=(ClampMin="0", AllowPreserveRatio, EditCondition="Type == ECEClonerEffectorType::Box", EditConditionHides))
 	FVector OuterExtent = FVector(200.f);
 
 	/** Plane spacing, everything inside this zone will be affected */
-	UPROPERTY(EditAnywhere, BlueprintReadWrite, Setter="SetPlaneSpacing", Getter="GetPlaneSpacing", Category="Type", meta=(ClampMin="0", EditCondition="Type == ECEClonerEffectorType::Plane", EditConditionHides))
+	UPROPERTY(EditAnywhere, BlueprintReadWrite, Setter, Getter, Category="Type", meta=(ClampMin="0", EditCondition="Type == ECEClonerEffectorType::Plane", EditConditionHides))
 	float PlaneSpacing = 200.f;
 
 	/** Radial angle in degree, everything within the angle will be affected */
@@ -547,102 +597,114 @@ protected:
 	UPROPERTY(EditAnywhere, BlueprintReadWrite, Setter, Getter, Category="Type", meta=(ClampMin="0", EditCondition="Type == ECEClonerEffectorType::Radial", EditConditionHides))
 	float RadialMaxRadius = 1000.f;
 
+	/** Main torus radius from center to the edge where inner and outer tube will be revolved */
+	UPROPERTY(EditAnywhere, BlueprintReadWrite, Setter, Getter, Category="Type", meta=(ClampMin="0", EditCondition="Type == ECEClonerEffectorType::Torus", EditConditionHides))
+	float TorusRadius = 250.f;
+
+	/** Minimum revolved radius for the torus effect, clones contained inside will be affected with a maximum weight */
+	UPROPERTY(EditAnywhere, BlueprintReadWrite, Setter, Getter, Category="Type", meta=(ClampMin="0", EditCondition="Type == ECEClonerEffectorType::Torus", EditConditionHides))
+	float TorusInnerRadius = 50.f;
+
+	/** Maximum revolved radius for the torus effect, clones outside of it will not be affected */
+	UPROPERTY(EditAnywhere, BlueprintReadWrite, Setter, Getter, Category="Type", meta=(ClampMin="0", EditCondition="Type == ECEClonerEffectorType::Torus", EditConditionHides))
+	float TorusOuterRadius = 200.f;
+
 	/** Mode of effector for each clones instances */
-	UPROPERTY(EditAnywhere, BlueprintReadWrite, Setter="SetMode", Getter="GetMode", Category="Mode")
+	UPROPERTY(EditAnywhere, BlueprintReadWrite, Setter, Getter, Category="Mode")
 	ECEClonerEffectorMode Mode = ECEClonerEffectorMode::Default;
 
 	/** Offset applied on affected clones */
-	UPROPERTY(EditAnywhere, BlueprintReadWrite, Setter="SetOffset", Getter="GetOffset", Category="Mode", meta=(EditCondition="Mode == ECEClonerEffectorMode::Default", EditConditionHides))
+	UPROPERTY(EditAnywhere, BlueprintReadWrite, Setter, Getter, Category="Mode", meta=(EditCondition="Mode == ECEClonerEffectorMode::Default", EditConditionHides))
 	FVector Offset = FVector::ZeroVector;
 
 	/** Rotation applied on affected clones */
-	UPROPERTY(EditAnywhere, BlueprintReadWrite, Setter="SetRotation", Getter="GetRotation", Category="Mode", meta=(ClampMin="-180", ClampMax="180", UIMin="-180", UIMax="180", EditCondition="Mode == ECEClonerEffectorMode::Default", EditConditionHides))
+	UPROPERTY(EditAnywhere, BlueprintReadWrite, Setter, Getter, Category="Mode", meta=(ClampMin="-180", ClampMax="180", UIMin="-180", UIMax="180", EditCondition="Mode == ECEClonerEffectorMode::Default", EditConditionHides))
 	FRotator Rotation = FRotator::ZeroRotator;
 
 	/** Scale applied on affected clones */
-	UPROPERTY(EditAnywhere, BlueprintReadWrite, Setter="SetScale", Getter="GetScale", Category="Mode", meta=(ClampMin="0", AllowPreserveRatio, Delta="0.01", EditCondition="Mode == ECEClonerEffectorMode::Default", EditConditionHides))
+	UPROPERTY(EditAnywhere, BlueprintReadWrite, Setter, Getter, Category="Mode", meta=(ClampMin="0", AllowPreserveRatio, Delta="0.01", EditCondition="Mode == ECEClonerEffectorMode::Default", EditConditionHides))
 	FVector Scale = FVector::OneVector;
 
 	/** The actor to track when mode is set to target */
-	UPROPERTY(EditAnywhere, BlueprintReadWrite, Category="Mode", meta=(DisplayName="TargetActor", EditCondition="Mode == ECEClonerEffectorMode::Target", EditConditionHides))
+	UPROPERTY(EditAnywhere, BlueprintReadWrite, Setter, Getter, Category="Mode", meta=(DisplayName="TargetActor", EditCondition="Mode == ECEClonerEffectorMode::Target", EditConditionHides))
 	TWeakObjectPtr<AActor> TargetActorWeak = nullptr;
 
 	UPROPERTY()
 	TWeakObjectPtr<AActor> InternalTargetActorWeak = nullptr;
 
 	/** Amplitude of the noise field for location */
-	UPROPERTY(EditAnywhere, BlueprintReadWrite, Setter="SetLocationStrength", Getter="GetLocationStrength", Category="Mode", meta=(EditCondition="Mode == ECEClonerEffectorMode::NoiseField", EditConditionHides))
+	UPROPERTY(EditAnywhere, BlueprintReadWrite, Setter, Getter, Category="Mode", meta=(EditCondition="Mode == ECEClonerEffectorMode::NoiseField", EditConditionHides))
 	FVector LocationStrength = FVector::ZeroVector;
 
 	/** Amplitude of the noise field for rotation */
-	UPROPERTY(EditAnywhere, BlueprintReadWrite, Category="Mode", meta=(EditCondition="Mode == ECEClonerEffectorMode::NoiseField", EditConditionHides))
+	UPROPERTY(EditAnywhere, BlueprintReadWrite, Setter, Getter, Category="Mode", meta=(EditCondition="Mode == ECEClonerEffectorMode::NoiseField", EditConditionHides))
 	FRotator RotationStrength = FRotator::ZeroRotator;
 
 	/** Amplitude of the noise field for scale */
-	UPROPERTY(EditAnywhere, BlueprintReadWrite, Category="Mode", meta=(EditCondition="Mode == ECEClonerEffectorMode::NoiseField", EditConditionHides))
+	UPROPERTY(EditAnywhere, BlueprintReadWrite, Setter, Getter, Category="Mode", meta=(ClampMin="0", AllowPreserveRatio, Delta="0.01", EditCondition="Mode == ECEClonerEffectorMode::NoiseField", EditConditionHides))
 	FVector ScaleStrength = FVector::OneVector;
 
 	/** Panning to offset the noise field sampling */
-	UPROPERTY(EditAnywhere, BlueprintReadWrite, Setter="SetPan", Getter="GetPan", Category="Mode", meta=(EditCondition="Mode == ECEClonerEffectorMode::NoiseField", EditConditionHides))
+	UPROPERTY(EditAnywhere, BlueprintReadWrite, Setter, Getter, Category="Mode", meta=(EditCondition="Mode == ECEClonerEffectorMode::NoiseField", EditConditionHides))
 	FVector Pan = FVector::ZeroVector;
 
 	/** Intensity of the noise field */
-	UPROPERTY(EditAnywhere, BlueprintReadWrite, Setter="SetFrequency", Getter="GetFrequency", Category="Mode", meta=(ClampMin="0", EditCondition="Mode == ECEClonerEffectorMode::NoiseField", EditConditionHides))
+	UPROPERTY(EditAnywhere, BlueprintReadWrite, Setter, Getter, Category="Mode", meta=(ClampMin="0", EditCondition="Mode == ECEClonerEffectorMode::NoiseField", EditConditionHides))
 	float Frequency = 0.5f;
 
 	/** Enable orientation force to allow each clone instance to rotate around its pivot */
 	UPROPERTY(EditAnywhere, BlueprintReadWrite, Setter="SetOrientationForceEnabled", Getter="GetOrientationForceEnabled", Category="Force")
 	bool bOrientationForceEnabled = false;
 
-	UPROPERTY(EditAnywhere, BlueprintReadWrite, Setter="SetOrientationForceRate", Getter="GetOrientationForceRate", Category="Force", meta=(Delta="0.0001", ClampMin="0", EditCondition="bOrientationForceEnabled", EditConditionHides))
+	UPROPERTY(EditAnywhere, BlueprintReadWrite, Setter, Getter, Category="Force", meta=(Delta="0.0001", ClampMin="0", EditCondition="bOrientationForceEnabled", EditConditionHides))
 	float OrientationForceRate = 1.f;
 
-	UPROPERTY(EditAnywhere, BlueprintReadWrite, Setter="SetOrientationForceMin", Getter="GetOrientationForceMin", Category="Force", meta=(EditCondition="bOrientationForceEnabled", EditConditionHides))
+	UPROPERTY(EditAnywhere, BlueprintReadWrite, Setter, Getter, Category="Force", meta=(EditCondition="bOrientationForceEnabled", EditConditionHides))
 	FVector OrientationForceMin = FVector(-0.1f);
 
-	UPROPERTY(EditAnywhere, BlueprintReadWrite, Setter="SetOrientationForceMax", Getter="GetOrientationForceMax", Category="Force", meta=(EditCondition="bOrientationForceEnabled", EditConditionHides))
+	UPROPERTY(EditAnywhere, BlueprintReadWrite, Setter, Getter, Category="Force", meta=(EditCondition="bOrientationForceEnabled", EditConditionHides))
 	FVector OrientationForceMax = FVector(0.1f);
 
 	/** Enable vortex force to allow each clone instance to rotate around a specific axis on the cloner pivot */
 	UPROPERTY(EditAnywhere, BlueprintReadWrite, Setter="SetVortexForceEnabled", Getter="GetVortexForceEnabled", Category="Force")
 	bool bVortexForceEnabled = false;
 
-	UPROPERTY(EditAnywhere, BlueprintReadWrite, Setter="SetVortexForceAmount", Getter="GetVortexForceAmount", Category="Force", meta=(EditCondition="bVortexForceEnabled", EditConditionHides))
+	UPROPERTY(EditAnywhere, BlueprintReadWrite, Setter, Getter, Category="Force", meta=(EditCondition="bVortexForceEnabled", EditConditionHides))
 	float VortexForceAmount = 10000.f;
 
-	UPROPERTY(EditAnywhere, BlueprintReadWrite, Setter="SetVortexForceAxis", Getter="GetVortexForceAxis", Category="Force", meta=(EditCondition="bVortexForceEnabled", EditConditionHides))
+	UPROPERTY(EditAnywhere, BlueprintReadWrite, Setter, Getter, Category="Force", meta=(EditCondition="bVortexForceEnabled", EditConditionHides))
 	FVector VortexForceAxis = FVector::ZAxisVector;
 
 	/** Enable curl noise force to allow each clone instance to add random location variation */
 	UPROPERTY(EditAnywhere, BlueprintReadWrite, Setter="SetCurlNoiseForceEnabled", Getter="GetCurlNoiseForceEnabled", Category="Force")
 	bool bCurlNoiseForceEnabled = false;
 
-	UPROPERTY(EditAnywhere, BlueprintReadWrite, Setter="SetCurlNoiseForceStrength", Getter="GetCurlNoiseForceStrength", Category="Force", meta=(EditCondition="bCurlNoiseForceEnabled", EditConditionHides))
+	UPROPERTY(EditAnywhere, BlueprintReadWrite, Setter, Getter, Category="Force", meta=(EditCondition="bCurlNoiseForceEnabled", EditConditionHides))
 	float CurlNoiseForceStrength = 1000.f;
 
-	UPROPERTY(EditAnywhere, BlueprintReadWrite, Setter="SetCurlNoiseForceFrequency", Getter="GetCurlNoiseForceFrequency", Category="Force", meta=(EditCondition="bCurlNoiseForceEnabled", EditConditionHides))
+	UPROPERTY(EditAnywhere, BlueprintReadWrite, Setter, Getter, Category="Force", meta=(EditCondition="bCurlNoiseForceEnabled", EditConditionHides))
 	float CurlNoiseForceFrequency = 10.f;
 
 	/** Enable attraction force to allow each clone instances to gravitate toward a location */
 	UPROPERTY(EditAnywhere, BlueprintReadWrite, Setter="SetAttractionForceEnabled", Getter="GetAttractionForceEnabled", Category="Force")
 	bool bAttractionForceEnabled = false;
 
-	UPROPERTY(EditAnywhere, BlueprintReadWrite, Setter="SetAttractionForceStrength", Getter="GetAttractionForceStrength", Category="Force", meta=(EditCondition="bAttractionForceEnabled", EditConditionHides))
+	UPROPERTY(EditAnywhere, BlueprintReadWrite, Setter, Getter, Category="Force", meta=(EditCondition="bAttractionForceEnabled", EditConditionHides))
 	float AttractionForceStrength = 1000.f;
 
-	UPROPERTY(EditAnywhere, BlueprintReadWrite, Setter="SetAttractionForceFalloff", Getter="GetAttractionForceFalloff", Category="Force", meta=(EditCondition="bAttractionForceEnabled", EditConditionHides))
+	UPROPERTY(EditAnywhere, BlueprintReadWrite, Setter, Getter, Category="Force", meta=(EditCondition="bAttractionForceEnabled", EditConditionHides))
 	float AttractionForceFalloff = 0.1f;
 
 	/** Enable gravity force to pull particles based on an acceleration vector */
 	UPROPERTY(EditAnywhere, BlueprintReadWrite, Setter="SetGravityForceEnabled", Getter="GetGravityForceEnabled", Category="Force")
 	bool bGravityForceEnabled = false;
 
-	UPROPERTY(EditAnywhere, BlueprintReadWrite, Setter="SetGravityForceAcceleration", Getter="GetGravityForceAcceleration", Category="Force", meta=(EditCondition="bGravityForceEnabled", EditConditionHides))
+	UPROPERTY(EditAnywhere, BlueprintReadWrite, Setter, Getter, Category="Force", meta=(EditCondition="bGravityForceEnabled", EditConditionHides))
 	FVector GravityForceAcceleration = FVector(0, 0, -980.f);
 
-	/** Thickness of components visualizers */
-	UPROPERTY(EditAnywhere, BlueprintReadWrite, AdvancedDisplay, Setter="SetVisualizerThickness", Getter="GetVisualizerThickness", Category="Effector", meta=(ClampMin="0.1", ClampMax="10.0", Delta="0.1"))
-	float VisualizerThickness = 0.5f;
+	/** Opacity of components visualizers */
+	UPROPERTY(EditAnywhere, BlueprintReadWrite, AdvancedDisplay, Setter, Getter, Category="Effector", meta=(ClampMin="0", ClampMax="1.0"))
+	float VisualizerOpacity = 0.1f;
 
 #if WITH_EDITORONLY_DATA
 	/** Toggle the sprite to visualize and click on this effector */
@@ -656,43 +718,32 @@ private:
 	UFUNCTION()
 	void OnTargetActorDestroyed(AActor* InActor);
 
-	/** Sphere component for sphere maximum weight */
-	UPROPERTY()
-	TObjectPtr<USphereComponent> InnerSphereComponent;
-
-	/** Sphere component for sphere minimum weight */
-	UPROPERTY()
-	TObjectPtr<USphereComponent> OuterSphereComponent;
-
-	/** Box component for box maximum weight */
-	UPROPERTY()
-	TObjectPtr<UBoxComponent> InnerBoxComponent;
-
-	/** Box component for box minimum weight */
-	UPROPERTY()
-	TObjectPtr<UBoxComponent> OuterBoxComponent;
+	/** Visualizer component for inner effector type */
+	UPROPERTY(Transient)
+	TObjectPtr<UDynamicMeshComponent> InnerVisualizerComponent;
 
-	/** Plane components for plane minimum weight */
-	UPROPERTY()
-	TObjectPtr<UBoxComponent> InnerPlaneComponent;
+	/** Visualizer component for outer effector type */
+	UPROPERTY(Transient)
+	TObjectPtr<UDynamicMeshComponent> OuterVisualizerComponent;
 
-	/** Plane components for plane maximum weight */
-	UPROPERTY()
-	TObjectPtr<UBoxComponent> OuterPlaneComponent;
+	/** Dynamic material for the inner visualizer */
+	UPROPERTY(Transient)
+	TObjectPtr<UMaterialInstanceDynamic> InnerVisualizerMaterial;
 
-	/** Radial components for begin angle */
-	UPROPERTY()
-	TObjectPtr<UArrowComponent> BeginRadialComponent;
+	/** Dynamic material for the outer visualizer */
+	UPROPERTY(Transient)
+	TObjectPtr<UMaterialInstanceDynamic> OuterVisualizerMaterial;
 
-	/** Radial components for end angle */
-	UPROPERTY()
-	TObjectPtr<UArrowComponent> EndRadialComponent;
+	/** Used to hold data related to visualizer for updates */
+	UPROPERTY(Transient)
+	FInstancedPropertyBag VisualizerData;
 
 	/** Internal cloners array, deprecated since it will be moved to the cloners and emptied out on post load */
 	UPROPERTY(NonTransactional)
 	TSet<TWeakObjectPtr<ACEClonerActor>> InternalCloners;
 
 	/** Transient effector channel data */
+	UPROPERTY(VisibleInstanceOnly, Transient, DuplicateTransient, TextExportTransient, NonTransactional, AdvancedDisplay, Category="Effector")
 	FCEClonerEffectorChannelData ChannelData;
 
 #if WITH_EDITOR
