MotionDesign : ClonerEffector
- Adding project settings for cloner and effector to control visualizer color, cloner spawn behavior, cloner default mesh, cloner default material

#virtualized

diff --git a/Engine/Plugins/Experimental/ClonerEffector/Source/ClonerEffector/ClonerEffector.Build.cs b/Engine/Plugins/Experimental/ClonerEffector/Source/ClonerEffector/ClonerEffector.Build.cs
--- a/Engine/Plugins/Experimental/ClonerEffector/Source/ClonerEffector/ClonerEffector.Build.cs
+++ b/Engine/Plugins/Experimental/ClonerEffector/Source/ClonerEffector/ClonerEffector.Build.cs
@@ -23,6 +23,7 @@ public class ClonerEffector : ModuleRules
 		PrivateDependencyModuleNames.AddRange(
 			new string[]
 			{
+				"DeveloperSettings",
 				"DynamicMesh",
 				"GeometryCore",
 				"GeometryFramework",
diff --git a/Engine/Plugins/Experimental/ClonerEffector/Source/ClonerEffector/Private/Cloner/CEClonerActor.cpp b/Engine/Plugins/Experimental/ClonerEffector/Source/ClonerEffector/Private/Cloner/CEClonerActor.cpp
--- a/Engine/Plugins/Experimental/ClonerEffector/Source/ClonerEffector/Private/Cloner/CEClonerActor.cpp
+++ b/Engine/Plugins/Experimental/ClonerEffector/Source/ClonerEffector/Private/Cloner/CEClonerActor.cpp
@@ -18,6 +18,7 @@
 #include "NiagaraDataInterfaceSkeletalMesh.h"
 #include "NiagaraMeshRendererProperties.h"
 #include "NiagaraSystem.h"
+#include "Settings/CEClonerEffectorSettings.h"
 #include "Subsystems/CEClonerSubsystem.h"
 #include "Subsystems/CEEffectorSubsystem.h"
 #include "UObject/ConstructorHelpers.h"
@@ -47,7 +48,7 @@ ACEClonerActor::ACEClonerActor()
 	LifetimeScaleCurve.AddKey(1, 0.f);
 
 	// Default override material
-	static ConstructorHelpers::FObjectFinder<UMaterialInterface> DefaultMaterialFinder(DefaultMaterialPath);
+	static ConstructorHelpers::FObjectFinder<UMaterialInterface> DefaultMaterialFinder(UCEClonerEffectorSettings::DefaultMaterialPath);
 	OverrideMaterial = DefaultMaterialFinder.Object;
 
 	if (!IsTemplate())
@@ -1282,7 +1283,7 @@ void ACEClonerActor::OnOverrideMaterialChanged()
 	if (ClonerComponent)
 	{
 		UMaterialInterface* OverrideMeshesMaterial = bVisualizeEffectors
-			? LoadObject<UMaterialInterface>(nullptr, DefaultMaterialPath)
+			? LoadObject<UMaterialInterface>(nullptr, UCEClonerEffectorSettings::DefaultMaterialPath)
 			: OverrideMaterial.Get();
 
 		ClonerComponent->SetOverrideMeshesMaterial(OverrideMeshesMaterial);
@@ -1364,6 +1365,12 @@ void ACEClonerActor::SpawnDefaultActorAttached()
 
 	bSpawnDefaultActorAttached = false;
 
+	const UCEClonerEffectorSettings* ClonerEffectorSettings = GetDefault<UCEClonerEffectorSettings>();
+	if (!ClonerEffectorSettings || !ClonerEffectorSettings->bSpawnDefaultActorAttached)
+	{
+		return;
+	}
+
 	// Only spawn if world is valid and not a preview actor
 	UWorld* World = GetWorld();
 	if (!World || bIsEditorPreviewActor)
@@ -1382,12 +1389,8 @@ void ACEClonerActor::SpawnDefaultActorAttached()
 		return;
 	}
 
-	// Find or load cube mesh
-	constexpr const TCHAR* DefaultStaticMeshPath = TEXT("/Script/Engine.StaticMesh'/Engine/BasicShapes/Cube.Cube'");
-	UStaticMesh* DefaultStaticMesh = LoadObject<UStaticMesh>(nullptr, DefaultStaticMeshPath);
-
-	// Find or load default material
-	UMaterialInterface* DefaultMaterial = LoadObject<UMaterialInterface>(nullptr, DefaultMaterialPath);
+	UStaticMesh* DefaultStaticMesh = ClonerEffectorSettings->DefaultStaticMesh.Get();
+	UMaterialInterface* DefaultMaterial = ClonerEffectorSettings->DefaultMaterial.Get();
 
 	if (!DefaultStaticMesh || !DefaultMaterial)
 	{
@@ -1409,7 +1412,7 @@ void ACEClonerActor::SpawnDefaultActorAttached()
 		DefaultActorAttached->SetMobility(EComponentMobility::Movable);
 		DefaultActorAttached->AttachToActor(this, FAttachmentTransformRules::KeepWorldTransform);
 
-		FActorLabelUtilities::SetActorLabelUnique(DefaultActorAttached, TEXT("DefaultCube"));
+		FActorLabelUtilities::SetActorLabelUnique(DefaultActorAttached, TEXT("DefaultClone"));
 	}
 }
 
diff --git a/Engine/Plugins/Experimental/ClonerEffector/Source/ClonerEffector/Private/Effector/CEEffectorActor.cpp b/Engine/Plugins/Experimental/ClonerEffector/Source/ClonerEffector/Private/Effector/CEEffectorActor.cpp
--- a/Engine/Plugins/Experimental/ClonerEffector/Source/ClonerEffector/Private/Effector/CEEffectorActor.cpp
+++ b/Engine/Plugins/Experimental/ClonerEffector/Source/ClonerEffector/Private/Effector/CEEffectorActor.cpp
@@ -5,12 +5,16 @@
 #include "Cloner/CEClonerActor.h"
 #include "Components/DynamicMeshComponent.h"
 #include "Effector/CEEffectorComponent.h"
+#include "Math/Vector.h"
+#include "Subsystems/CEEffectorSubsystem.h"
+
+#if WITH_EDITOR
 #include "GeometryScript/MeshPrimitiveFunctions.h"
 #include "Materials/MaterialInstanceDynamic.h"
-#include "Math/Vector.h"
 #include "PropertyBag.h"
-#include "Subsystems/CEEffectorSubsystem.h"
+#include "Settings/CEClonerEffectorSettings.h"
 #include "UObject/ConstructorHelpers.h"
+#endif
 
 ACEEffectorActor::FOnEffectorIdentifierChanged ACEEffectorActor::OnEffectorRefreshClonerDelegate;
 
@@ -30,10 +34,9 @@ ACEEffectorActor::ACEEffectorActor()
 	InnerVisualizerComponent->SetHiddenInGame(true);
 	InnerVisualizerComponent->SetTranslucentSortPriority(0);
 #if WITH_EDITOR
-	// Do not show bounding box around cloner for better visibility
 	InnerVisualizerComponent->SetIsVisualizationComponent(true);
 #endif
-	InnerVisualizerComponent->bIsEditorOnly = false;
+	InnerVisualizerComponent->bIsEditorOnly = true;
 	InnerVisualizerComponent->SetupAttachment(SceneComponent);
 
 	// Outer Visualizer Component
@@ -41,12 +44,12 @@ ACEEffectorActor::ACEEffectorActor()
 	OuterVisualizerComponent->SetHiddenInGame(true);
 	OuterVisualizerComponent->SetTranslucentSortPriority(1);
 #if WITH_EDITOR
-	// Do not show bounding box around cloner for better visibility
 	OuterVisualizerComponent->SetIsVisualizationComponent(true);
 #endif
-	OuterVisualizerComponent->bIsEditorOnly = false;
+	OuterVisualizerComponent->bIsEditorOnly = true;
 	OuterVisualizerComponent->SetupAttachment(SceneComponent);
 
+#if WITH_EDITOR
 	static const ConstructorHelpers::FObjectFinder<UMaterialInterface> VisualizerMaterial(TEXT("/Script/Engine.Material'/ClonerEffector/Materials/M_EffectorVisualizer.M_EffectorVisualizer'"));
 
 	if (VisualizerMaterial.Succeeded())
@@ -56,22 +59,24 @@ ACEEffectorActor::ACEEffectorActor()
 			this
 		);
 
-		InnerVisualizerMaterial->SetVectorParameterValue(VisualizerColorName, FLinearColor::Red);
-		InnerVisualizerMaterial->SetScalarParameterValue(VisualizerOpacityName, VisualizerOpacity);
-
 		OuterVisualizerMaterial = UMaterialInstanceDynamic::Create(
 			VisualizerMaterial.Object,
 			this
 		);
-
-		OuterVisualizerMaterial->SetVectorParameterValue(VisualizerColorName, FLinearColor::Blue);
-		OuterVisualizerMaterial->SetScalarParameterValue(VisualizerOpacityName, VisualizerOpacity);
 	}
+#endif
 
 	if (!IsTemplate())
 	{
 		SceneComponent->TransformUpdated.AddUObject(this, &ACEEffectorActor::OnEffectorTransformed);
 		UCEEffectorSubsystem::OnSubsystemInitializedDelegate.AddUObject(this, &ACEEffectorActor::OnEffectorSubsystemInitialized);
+
+#if WITH_EDITOR
+		if (UCEClonerEffectorSettings* ClonerEffectorSettings = GetMutableDefault<UCEClonerEffectorSettings>())
+		{
+			ClonerEffectorSettings->OnSettingChanged().AddUObject(this, &ACEEffectorActor::OnEffectorDeveloperSettingsChanged);
+		}
+#endif
 	}
 }
 
@@ -86,8 +91,8 @@ TCEPropertyChangeDispatcher<ACEEffectorActor> ACEEffectorActor::PropertyChangeDi
 	/** Effector */
 	{ GET_MEMBER_NAME_CHECKED(ACEEffectorActor, bEnabled), &ACEEffectorActor::OnEnabledChanged },
 	{ GET_MEMBER_NAME_CHECKED(ACEEffectorActor, Magnitude), &ACEEffectorActor::OnMagnitudeChanged },
-	{ GET_MEMBER_NAME_CHECKED(ACEEffectorActor, VisualizerOpacity), &ACEEffectorActor::OnVisualizerOpacityChanged },
-	{ GET_MEMBER_NAME_CHECKED(ACEEffectorActor, bVisualizerSpriteVisible), &ACEEffectorActor::OnVisualizerSpriteVisibleChanged },
+	{ GET_MEMBER_NAME_CHECKED(ACEEffectorActor, bVisualizerComponentVisible), &ACEEffectorActor::OnVisualizerOptionsChanged },
+	{ GET_MEMBER_NAME_CHECKED(ACEEffectorActor, bVisualizerSpriteVisible), &ACEEffectorActor::OnVisualizerOptionsChanged },
 	{ GET_MEMBER_NAME_CHECKED(ACEEffectorActor, Color), &ACEEffectorActor::OnColorChanged },
 	/** Type */
 	{ GET_MEMBER_NAME_CHECKED(ACEEffectorActor, Type), &ACEEffectorActor::OnTypeChanged },
@@ -493,22 +498,6 @@ void ACEEffectorActor::SetEnabled(bool bInEnable)
 	OnEnabledChanged();
 }
 
-void ACEEffectorActor::SetVisualizerOpacity(float InOpacity)
-{
-	if (FMath::IsNearlyEqual(VisualizerOpacity, InOpacity))
-	{
-		return;
-	}
-
-	if (VisualizerOpacity < 0.f || VisualizerOpacity > 1.f)
-	{
-		return;
-	}
-
-	VisualizerOpacity = InOpacity;
-	OnVisualizerOpacityChanged();
-}
-
 void ACEEffectorActor::SetMode(ECEClonerEffectorMode InMode)
 {
 	if (Mode == InMode)
@@ -762,6 +751,17 @@ void ACEEffectorActor::SetAttractionForceEnabled(bool bInForceEnabled)
 }
 
 #if WITH_EDITOR
+void ACEEffectorActor::SetVisualizerComponentVisible(bool bInVisible)
+{
+	if (bVisualizerComponentVisible == bInVisible)
+	{
+		return;
+	}
+
+	bVisualizerComponentVisible = bInVisible;
+	OnVisualizerOptionsChanged();
+}
+
 void ACEEffectorActor::SetVisualizerSpriteVisible(bool bInVisible)
 {
 	if (bVisualizerSpriteVisible == bInVisible)
@@ -770,7 +770,7 @@ void ACEEffectorActor::SetVisualizerSpriteVisible(bool bInVisible)
 	}
 
 	bVisualizerSpriteVisible = bInVisible;
-	OnVisualizerSpriteVisibleChanged();
+	OnVisualizerOptionsChanged();
 }
 #endif
 
@@ -809,11 +809,13 @@ void ACEEffectorActor::OnEnabledChanged()
 	else // Disabled
 	{
 		OnEffectorDisabled();
+
 		// Hide visualization components
 		OnTypeChanged();
-		// Editor
-		OnVisualizerOpacityChanged();
-		OnVisualizerSpriteVisibleChanged();
+
+#if WITH_EDITOR
+		OnVisualizerOptionsChanged();
+#endif
 	}
 }
 
@@ -856,9 +858,10 @@ void ACEEffectorActor::OnEffectorChanged()
 		OnMagnitudeChanged();
 		OnForceOptionsChanged();
 		OnColorChanged();
-		// Editor
-		OnVisualizerOpacityChanged();
-		OnVisualizerSpriteVisibleChanged();
+
+#if WITH_EDITOR
+		OnVisualizerOptionsChanged();
+#endif
 	}
 }
 
@@ -887,6 +890,7 @@ void ACEEffectorActor::OnBoxChanged()
 	ChannelData.InnerExtent = GetInnerExtent();
 	ChannelData.OuterExtent = GetOuterExtent();
 
+#if WITH_EDITOR
 	// Update visualizer
 	static const FName BoxInnerExtentName = GET_MEMBER_NAME_CHECKED(ACEEffectorActor, InnerExtent);
 	static const FName BoxOuterExtentName = GET_MEMBER_NAME_CHECKED(ACEEffectorActor, OuterExtent);
@@ -918,6 +922,7 @@ void ACEEffectorActor::OnBoxChanged()
 	VisualizerData.AddProperty(BoxOuterExtentName, EPropertyBagPropertyType::Struct);
 	VisualizerData.SetValueStruct(BoxInnerExtentName, InnerExtent);
 	VisualizerData.SetValueStruct(BoxOuterExtentName, OuterExtent);
+#endif
 }
 
 void ACEEffectorActor::OnSphereChanged()
@@ -933,6 +938,7 @@ void ACEEffectorActor::OnSphereChanged()
 	ChannelData.InnerExtent = FVector(GetInnerRadius());
 	ChannelData.OuterExtent = FVector(GetOuterRadius());
 
+#if WITH_EDITOR
 	// Update visualizer
 	static const FName SphereInnerRadiusName = GET_MEMBER_NAME_CHECKED(ACEEffectorActor, InnerRadius);
 	static const FName SphereOuterRadiusName = GET_MEMBER_NAME_CHECKED(ACEEffectorActor, OuterRadius);
@@ -964,6 +970,7 @@ void ACEEffectorActor::OnSphereChanged()
 	VisualizerData.AddProperty(SphereOuterRadiusName, EPropertyBagPropertyType::Float);
 	VisualizerData.SetValueFloat(SphereInnerRadiusName, InnerRadius);
 	VisualizerData.SetValueFloat(SphereOuterRadiusName, OuterRadius);
+#endif
 }
 
 void ACEEffectorActor::OnPlaneChanged()
@@ -976,6 +983,7 @@ void ACEEffectorActor::OnPlaneChanged()
 	ChannelData.InnerExtent = FVector::LeftVector;
 	ChannelData.OuterExtent = FVector(PlaneSpacing);
 
+#if WITH_EDITOR
 	// Update visualizer
 	static const FName PlaneSpacingName = GET_MEMBER_NAME_CHECKED(ACEEffectorActor, PlaneSpacing);
 	TValueOrError<float, EPropertyBagResult> PlaneSpacingValue = VisualizerData.GetValueFloat(PlaneSpacingName);
@@ -1001,6 +1009,7 @@ void ACEEffectorActor::OnPlaneChanged()
 	VisualizerData.Reset();
 	VisualizerData.AddProperty(PlaneSpacingName, EPropertyBagPropertyType::Float);
 	VisualizerData.SetValueFloat(PlaneSpacingName, PlaneSpacing);
+#endif
 }
 
 void ACEEffectorActor::OnRadialChanged()
@@ -1016,6 +1025,7 @@ void ACEEffectorActor::OnRadialChanged()
 	ChannelData.InnerExtent = FVector::LeftVector;
 	ChannelData.OuterExtent = FVector(RadialAngle, RadialMinRadius, RadialMaxRadius);
 
+#if WITH_EDITOR
 	// Update visualizer
 	static const FName RadialAngleName = GET_MEMBER_NAME_CHECKED(ACEEffectorActor, RadialAngle);
 	static const FName RadialMinRadiusName = GET_MEMBER_NAME_CHECKED(ACEEffectorActor, RadialMinRadius);
@@ -1051,6 +1061,7 @@ void ACEEffectorActor::OnRadialChanged()
 	VisualizerData.SetValueFloat(RadialAngleName, RadialAngle);
 	VisualizerData.SetValueFloat(RadialMinRadiusName, RadialMinRadius);
 	VisualizerData.SetValueFloat(RadialMaxRadiusName, RadialMaxRadius);
+#endif
 }
 
 void ACEEffectorActor::OnTorusChanged()
@@ -1067,6 +1078,7 @@ void ACEEffectorActor::OnTorusChanged()
 	ChannelData.InnerExtent = FVector::ZAxisVector;
 	ChannelData.OuterExtent = FVector(TorusInnerRadius, TorusOuterRadius, TorusRadius);
 
+#if WITH_EDITOR
 	// Update visualizer
 	static const FName TorusRadiusName = GET_MEMBER_NAME_CHECKED(ACEEffectorActor, TorusRadius);
 	static const FName TorusInnerRadiusName = GET_MEMBER_NAME_CHECKED(ACEEffectorActor, TorusInnerRadius);
@@ -1108,6 +1120,7 @@ void ACEEffectorActor::OnTorusChanged()
 	VisualizerData.SetValueFloat(TorusRadiusName, TorusRadius);
 	VisualizerData.SetValueFloat(TorusInnerRadiusName, TorusInnerRadius);
 	VisualizerData.SetValueFloat(TorusOuterRadiusName, TorusOuterRadius);
+#endif
 }
 
 void ACEEffectorActor::OnUnboundChanged()
@@ -1117,9 +1130,11 @@ void ACEEffectorActor::OnUnboundChanged()
 		return;
 	}
 
+#if WITH_EDITOR
 	// Clear visualizers
 	UpdateVisualizer(InnerVisualizerId, [](UDynamicMesh* InMesh){});
 	UpdateVisualizer(OuterVisualizerId, [](UDynamicMesh* InMesh){});
+#endif
 }
 
 void ACEEffectorActor::OnMagnitudeChanged()
@@ -1196,29 +1211,6 @@ void ACEEffectorActor::OnTargetActorDestroyed(AActor* InActor)
 	}
 }
 
-void ACEEffectorActor::UpdateVisualizer(int32 InVisualizerId, TFunction<void(UDynamicMesh*)> InMeshFunction) const
-{
-	UDynamicMeshComponent* MeshComponent = InVisualizerId == InnerVisualizerId ? InnerVisualizerComponent : OuterVisualizerComponent;
-
-	if (!MeshComponent)
-	{
-		return;
-	}
-
-	UDynamicMesh* DynamicMesh = MeshComponent->GetDynamicMesh();
-
-	DynamicMesh->EditMesh([](FDynamicMesh3& InMesh)
-	{
-		InMesh.Clear();
-	});
-
-	InMeshFunction(DynamicMesh);
-
-	// Apply material
-	UMaterialInstanceDynamic* VisualizerMaterial = InVisualizerId == InnerVisualizerId ? InnerVisualizerMaterial : OuterVisualizerMaterial;
-	MeshComponent->SetMaterial(0, VisualizerMaterial);
-}
-
 void ACEEffectorActor::OnNoiseFieldOptionsChanged()
 {
 	if (Mode != ECEClonerEffectorMode::NoiseField)
@@ -1256,29 +1248,72 @@ void ACEEffectorActor::OnTargetOptionsChanged()
 	ChannelData.ScaleDelta = FVector::OneVector;
 }
 
-void ACEEffectorActor::OnVisualizerOpacityChanged()
+#if WITH_EDITOR
+void ACEEffectorActor::OnVisualizerOptionsChanged()
 {
-	VisualizerOpacity = FMath::Clamp(VisualizerOpacity, 0.f, 1.f);
+	const UCEClonerEffectorSettings* ClonerEffectorSettings = GetDefault<UCEClonerEffectorSettings>();
+
+	if (!ClonerEffectorSettings)
+	{
+		return;
+	}
 
 	if (InnerVisualizerMaterial)
 	{
-		InnerVisualizerMaterial->SetScalarParameterValue(VisualizerOpacityName, VisualizerOpacity);
+		InnerVisualizerMaterial->SetVectorParameterValue(VisualizerColorName, ClonerEffectorSettings->VisualizerInnerColor);
 	}
 
 	if (OuterVisualizerMaterial)
 	{
-		OuterVisualizerMaterial->SetScalarParameterValue(VisualizerOpacityName, VisualizerOpacity);
+		OuterVisualizerMaterial->SetVectorParameterValue(VisualizerColorName, ClonerEffectorSettings->VisualizerOuterColor);
+	}
+
+	if (InnerVisualizerComponent)
+	{
+		InnerVisualizerComponent->SetVisibility(bVisualizerComponentVisible, false);
+	}
+
+	if (OuterVisualizerComponent)
+	{
+		OuterVisualizerComponent->SetVisibility(bVisualizerComponentVisible, false);
 	}
-}
 
-void ACEEffectorActor::OnVisualizerSpriteVisibleChanged()
-{
-#if WITH_EDITOR
 	UE::ClonerEffector::SetBillboardComponentSprite(this, TEXT("/Script/Engine.Texture2D'/ClonerEffector/Textures/T_EffectorIcon.T_EffectorIcon'"));
 	UE::ClonerEffector::SetBillboardComponentVisibility(this, bVisualizerSpriteVisible);
-#endif
 }
 
+void ACEEffectorActor::OnEffectorDeveloperSettingsChanged(UObject* InSettings, FPropertyChangedEvent& InEvent)
+{
+	OnVisualizerOptionsChanged();
+}
+
+void ACEEffectorActor::UpdateVisualizer(int32 InVisualizerId, TFunction<void(UDynamicMesh*)> InMeshFunction) const
+{
+	UDynamicMeshComponent* MeshComponent = InVisualizerId == InnerVisualizerId ? InnerVisualizerComponent : OuterVisualizerComponent;
+
+	if (!MeshComponent)
+	{
+		return;
+	}
+
+	UDynamicMesh* DynamicMesh = MeshComponent->GetDynamicMesh();
+
+	DynamicMesh->EditMesh([](FDynamicMesh3& InMesh)
+	{
+		InMesh.Clear();
+	});
+
+	if (bVisualizerComponentVisible)
+	{
+		InMeshFunction(DynamicMesh);
+	}
+
+	// Apply material
+	UMaterialInstanceDynamic* VisualizerMaterial = InVisualizerId == InnerVisualizerId ? InnerVisualizerMaterial : OuterVisualizerMaterial;
+	MeshComponent->SetMaterial(0, VisualizerMaterial);
+}
+#endif
+
 void ACEEffectorActor::OnColorChanged()
 {
 	ChannelData.Color = Color;
diff --git a/Engine/Plugins/Experimental/ClonerEffector/Source/ClonerEffector/Private/Settings/CEClonerEffectorSettings.cpp b/Engine/Plugins/Experimental/ClonerEffector/Source/ClonerEffector/Private/Settings/CEClonerEffectorSettings.cpp
new file mode 100644
--- /dev/null
+++ b/Engine/Plugins/Experimental/ClonerEffector/Source/ClonerEffector/Private/Settings/CEClonerEffectorSettings.cpp
@@ -0,0 +1,16 @@
+// Copyright Epic Games, Inc. All Rights Reserved.
+
+#include "Settings/CEClonerEffectorSettings.h"
+
+#include "Engine/StaticMesh.h"
+#include "Materials/MaterialInterface.h"
+#include "UObject/SoftObjectPath.h"
+
+UCEClonerEffectorSettings::UCEClonerEffectorSettings()
+{
+	CategoryName = TEXT("Motion Design");
+	SectionName = TEXT("Cloner & Effector");
+
+	DefaultStaticMesh = TSoftObjectPtr<UStaticMesh>(FSoftObjectPath(DefaultStaticMeshPath));
+	DefaultMaterial = TSoftObjectPtr<UMaterialInterface>(FSoftObjectPath(DefaultMaterialPath));
+}
diff --git a/Engine/Plugins/Experimental/ClonerEffector/Source/ClonerEffector/Private/Settings/CEClonerEffectorSettings.h b/Engine/Plugins/Experimental/ClonerEffector/Source/ClonerEffector/Private/Settings/CEClonerEffectorSettings.h
new file mode 100644
--- /dev/null
+++ b/Engine/Plugins/Experimental/ClonerEffector/Source/ClonerEffector/Private/Settings/CEClonerEffectorSettings.h
@@ -0,0 +1,44 @@
+// Copyright Epic Games, Inc. All Rights Reserved.
+
+#pragma once
+
+#include "Engine/DeveloperSettings.h"
+#include "UObject/SoftObjectPtr.h"
+#include "CEClonerEffectorSettings.generated.h"
+
+class UMaterialInterface;
+class UStaticMesh;
+struct FLinearColor;
+
+/** Settings for motion design Cloner and Effector plugin */
+UCLASS(Config=Engine, meta=(DisplayName="Cloner & Effector"))
+class UCEClonerEffectorSettings : public UDeveloperSettings
+{
+	GENERATED_BODY()
+
+public:
+	static constexpr TCHAR DefaultStaticMeshPath[] = TEXT("/Script/Engine.StaticMesh'/Engine/BasicShapes/Cube.Cube'");
+	static constexpr TCHAR DefaultMaterialPath[] = TEXT("/Script/Engine.Material'/ClonerEffector/Materials/DefaultClonerMaterial.DefaultClonerMaterial'");
+
+	UCEClonerEffectorSettings();
+
+	/** Inner visualizer color for effectors */
+	UPROPERTY(Config, EditAnywhere, Category="Effector")
+	FLinearColor VisualizerInnerColor = FLinearColor(255, 0, 0, 0.1);
+
+	/** Outer visualizer color for effectors */
+	UPROPERTY(Config, EditAnywhere, Category="Effector")
+	FLinearColor VisualizerOuterColor = FLinearColor(0, 0, 255, 0.1);
+
+	/** Spawns a default actor attached to the cloner on spawn */
+	UPROPERTY(Config, EditAnywhere, Category="Cloner")
+	bool bSpawnDefaultActorAttached = true;
+
+	/** Default static mesh used when spawning default actor attached */
+	UPROPERTY(Config, EditAnywhere, Category="Cloner")
+	TSoftObjectPtr<UStaticMesh> DefaultStaticMesh;
+
+	/** Default material used when spawning default actor attached */
+	UPROPERTY(Config, EditAnywhere, Category="Cloner")
+	TSoftObjectPtr<UMaterialInterface> DefaultMaterial;
+};
diff --git a/Engine/Plugins/Experimental/ClonerEffector/Source/ClonerEffector/Public/Cloner/CEClonerActor.h b/Engine/Plugins/Experimental/ClonerEffector/Source/ClonerEffector/Public/Cloner/CEClonerActor.h
--- a/Engine/Plugins/Experimental/ClonerEffector/Source/ClonerEffector/Public/Cloner/CEClonerActor.h
+++ b/Engine/Plugins/Experimental/ClonerEffector/Source/ClonerEffector/Public/Cloner/CEClonerActor.h
@@ -17,7 +17,7 @@ class UCEClonerLayoutBase;
 class UNiagaraDataInterfaceCurve;
 class USceneComponent;
 
-UCLASS(MinimalAPI, BlueprintType, HideCategories=(Rendering,Replication,Collision,HLOD,Physics,Networking,Input,Actor,Cooking,LevelInstance,Streaming), DisplayName = "Motion Design Cloner Actor")
+UCLASS(MinimalAPI, BlueprintType, HideCategories=(Rendering,Replication,Collision,HLOD,Physics,Networking,Input,Actor,Cooking,LevelInstance,Streaming,DataLayers), DisplayName = "Motion Design Cloner Actor")
 class ACEClonerActor : public AActor
 {
 	GENERATED_BODY()
@@ -30,7 +30,6 @@ class ACEClonerActor : public AActor
 
 public:
 	static inline const FString DefaultLabel = TEXT("Cloner");
-	static inline constexpr const TCHAR* DefaultMaterialPath = TEXT("/Script/Engine.Material'/ClonerEffector/Materials/DefaultClonerMaterial.DefaultClonerMaterial'");
 
 	ACEClonerActor();
 
diff --git a/Engine/Plugins/Experimental/ClonerEffector/Source/ClonerEffector/Public/Effector/CEEffectorActor.h b/Engine/Plugins/Experimental/ClonerEffector/Source/ClonerEffector/Public/Effector/CEEffectorActor.h
--- a/Engine/Plugins/Experimental/ClonerEffector/Source/ClonerEffector/Public/Effector/CEEffectorActor.h
+++ b/Engine/Plugins/Experimental/ClonerEffector/Source/ClonerEffector/Public/Effector/CEEffectorActor.h
@@ -11,7 +11,7 @@
 class ACEClonerActor;
 class UDynamicMeshComponent;
 
-UCLASS(MinimalAPI, BlueprintType, HideCategories=(Rendering,Replication,Collision,HLOD,Physics,Networking,Input,Actor,Cooking,LevelInstance), DisplayName = "Motion Design Effector Actor")
+UCLASS(MinimalAPI, BlueprintType, HideCategories=(Rendering,Replication,Collision,HLOD,Physics,Networking,Input,Actor,Cooking,LevelInstance,DataLayers), DisplayName = "Motion Design Effector Actor")
 class ACEEffectorActor : public AActor
 {
 	GENERATED_BODY()
@@ -411,16 +411,16 @@ public:
 		return Color;
 	}
 
+#if WITH_EDITOR
 	UFUNCTION(BlueprintCallable, Category="Effector")
-	CLONEREFFECTOR_API void SetVisualizerOpacity(float InOpacity);
+	CLONEREFFECTOR_API void SetVisualizerComponentVisible(bool bInVisible);
 
 	UFUNCTION(BlueprintPure, Category="Effector")
-	float GetVisualizerOpacity() const
+	bool GetVisualizerComponentVisible() const
 	{
-		return VisualizerOpacity;
+		return bVisualizerComponentVisible;
 	}
 
-#if WITH_EDITOR
 	UFUNCTION(BlueprintCallable, Category="Effector")
 	CLONEREFFECTOR_API void SetVisualizerSpriteVisible(bool bInVisible);
 
@@ -439,9 +439,6 @@ protected:
 	static constexpr int32 InnerVisualizerId = 0;
 	static constexpr int32 OuterVisualizerId = 1;
 	static constexpr TCHAR VisualizerColorName[] = TEXT("VisualizerColor");
-	static constexpr TCHAR VisualizerOpacityName[] = TEXT("VisualizerOpacity");
-
-	void UpdateVisualizer(int32 InVisualizerId, TFunction<void(UDynamicMesh*)> InMeshFunction) const;
 
 	//~ Begin AActor
 	virtual void Destroyed() override;
@@ -529,14 +526,19 @@ protected:
 	/** Called when force enabled state changed */
 	void OnForceEnabledChanged();
 
-	/** Update opacity of components of this effector */
-	void OnVisualizerOpacityChanged();
+	/** Update particle color affected by this effector */
+    void OnColorChanged();
 
-	/** Update sprite visibility of this effector */
-	void OnVisualizerSpriteVisibleChanged();
+#if WITH_EDITOR
+	/** Update visualizers options of this effector */
+	void OnVisualizerOptionsChanged();
 
-	/** Update particle color affected by this effector */
-	void OnColorChanged();
+	/** Called when developer settings are changed */
+	void OnEffectorDeveloperSettingsChanged(UObject* InSettings, FPropertyChangedEvent& InEvent);
+
+	/** Update the mesh of a component visualizer */
+	void UpdateVisualizer(int32 InVisualizerId, TFunction<void(UDynamicMesh*)> InMeshFunction) const;
+#endif
 
 	/** Update all types options */
 	void UpdateEffectorTypes();
@@ -702,11 +704,11 @@ protected:
 	UPROPERTY(EditAnywhere, BlueprintReadWrite, Setter, Getter, Category="Force", meta=(EditCondition="bGravityForceEnabled", EditConditionHides))
 	FVector GravityForceAcceleration = FVector(0, 0, -980.f);
 
-	/** Opacity of components visualizers */
-	UPROPERTY(EditAnywhere, BlueprintReadWrite, AdvancedDisplay, Setter, Getter, Category="Effector", meta=(ClampMin="0", ClampMax="1.0"))
-	float VisualizerOpacity = 0.1f;
-
 #if WITH_EDITORONLY_DATA
+	/** Visibility of the components visualizer */
+	UPROPERTY(EditAnywhere, BlueprintReadWrite, AdvancedDisplay, Category="Effector")
+	bool bVisualizerComponentVisible = true;
+
 	/** Toggle the sprite to visualize and click on this effector */
 	UPROPERTY(EditAnywhere, BlueprintReadWrite, AdvancedDisplay, Category="Effector")
 	bool bVisualizerSpriteVisible = true;
@@ -726,6 +728,7 @@ private:
 	UPROPERTY(Transient)
 	TObjectPtr<UDynamicMeshComponent> OuterVisualizerComponent;
 
+#if WITH_EDITORONLY_DATA
 	/** Dynamic material for the inner visualizer */
 	UPROPERTY(Transient)
 	TObjectPtr<UMaterialInstanceDynamic> InnerVisualizerMaterial;
@@ -737,6 +740,7 @@ private:
 	/** Used to hold data related to visualizer for updates */
 	UPROPERTY(Transient)
 	FInstancedPropertyBag VisualizerData;
+#endif
 
 	/** Internal cloners array, deprecated since it will be moved to the cloners and emptied out on post load */
 	UPROPERTY(NonTransactional)
