MotionDesign : ClonerEffector
- Removing deprecated property from cloner/effector
- Removing old deprecated structs

diff --git a/Engine/Plugins/Experimental/ClonerEffector/Source/ClonerEffector/Private/Cloner/CEClonerActor.cpp b/Engine/Plugins/Experimental/ClonerEffector/Source/ClonerEffector/Private/Cloner/CEClonerActor.cpp
--- a/Engine/Plugins/Experimental/ClonerEffector/Source/ClonerEffector/Private/Cloner/CEClonerActor.cpp
+++ b/Engine/Plugins/Experimental/ClonerEffector/Source/ClonerEffector/Private/Cloner/CEClonerActor.cpp
@@ -3,14 +3,7 @@
 #include "Cloner/CEClonerActor.h"
 
 #include "Cloner/CEClonerComponent.h"
-#include "Cloner/Layouts/CEClonerCircleLayout.h"
-#include "Cloner/Layouts/CEClonerCylinderLayout.h"
-#include "Cloner/Layouts/CEClonerGridLayout.h"
-#include "Cloner/Layouts/CEClonerHoneycombLayout.h"
-#include "Cloner/Layouts/CEClonerLineLayout.h"
-#include "Cloner/Layouts/CEClonerMeshLayout.h"
-#include "Cloner/Layouts/CEClonerSphereUniformLayout.h"
-#include "Cloner/Layouts/CEClonerSplineLayout.h"
+#include "Cloner/Layouts/CEClonerLayoutBase.h"
 #include "Effector/CEEffectorActor.h"
 #include "Engine/StaticMeshActor.h"
 #include "Materials/MaterialInterface.h"
@@ -183,114 +176,6 @@ void ACEClonerActor::PostLoad()
 {
 	Super::PostLoad();
 
-	// Migrate old properties to new layout system
-	if (!bDeprecatedPropertiesMigrated)
-	{
-		// Grid options
-		if (UCEClonerGridLayout* GridLayout = FindOrAddLayout<UCEClonerGridLayout>())
-		{
-			GridLayout->SetCountX(GridOptions_DEPRECATED.Count.X);
-			GridLayout->SetCountY(GridOptions_DEPRECATED.Count.Y);
-			GridLayout->SetCountZ(GridOptions_DEPRECATED.Count.Z);
-			GridLayout->SetSpacingX(GridOptions_DEPRECATED.Spacing.X);
-			GridLayout->SetSpacingY(GridOptions_DEPRECATED.Spacing.Y);
-			GridLayout->SetSpacingZ(GridOptions_DEPRECATED.Spacing.Z);
-			GridLayout->SetConstraint(GridOptions_DEPRECATED.Constraint);
-			GridLayout->SetInvertConstraint(GridOptions_DEPRECATED.bInvertConstraint);
-			GridLayout->SetSphereConstraint(GridOptions_DEPRECATED.SphereConstraint);
-			GridLayout->SetCylinderConstraint(GridOptions_DEPRECATED.CylinderConstraint);
-			GridLayout->SetTextureConstraint(GridOptions_DEPRECATED.TextureConstraint);
-		}
-
-		// Line options
-		if (UCEClonerLineLayout* LineLayout = FindOrAddLayout<UCEClonerLineLayout>())
-		{
-			LineLayout->SetCount(LineOptions_DEPRECATED.Count);
-			LineLayout->SetSpacing(LineOptions_DEPRECATED.Spacing);
-			LineLayout->SetAxis(LineOptions_DEPRECATED.Axis);
-			LineLayout->SetDirection(LineOptions_DEPRECATED.Direction);
-			LineLayout->SetRotation(LineOptions_DEPRECATED.Rotation);
-		}
-
-		// Circle options
-		if (UCEClonerCircleLayout* CircleLayout = FindOrAddLayout<UCEClonerCircleLayout>())
-		{
-			CircleLayout->SetCount(CircleOptions_DEPRECATED.Count);
-			CircleLayout->SetRadius(CircleOptions_DEPRECATED.Radius);
-			CircleLayout->SetAngleStart(CircleOptions_DEPRECATED.AngleStart);
-			CircleLayout->SetAngleRatio(CircleOptions_DEPRECATED.AngleRatio);
-			CircleLayout->SetOrientMesh(CircleOptions_DEPRECATED.bOrientMesh);
-			CircleLayout->SetPlane(CircleOptions_DEPRECATED.Plane);
-			CircleLayout->SetRotation(CircleOptions_DEPRECATED.Rotation);
-			CircleLayout->SetScale(CircleOptions_DEPRECATED.Scale);
-		}
-
-		// Cylinder options
-		if (UCEClonerCylinderLayout* CylinderLayout = FindOrAddLayout<UCEClonerCylinderLayout>())
-		{
-			CylinderLayout->SetBaseCount(CylinderOptions_DEPRECATED.BaseCount);
-			CylinderLayout->SetHeightCount(CylinderOptions_DEPRECATED.HeightCount);
-			CylinderLayout->SetHeight(CylinderOptions_DEPRECATED.Height);
-			CylinderLayout->SetRadius(CylinderOptions_DEPRECATED.Radius);
-			CylinderLayout->SetAngleStart(CylinderOptions_DEPRECATED.AngleStart);
-			CylinderLayout->SetAngleRatio(CylinderOptions_DEPRECATED.AngleRatio);
-			CylinderLayout->SetOrientMesh(CylinderOptions_DEPRECATED.bOrientMesh);
-			CylinderLayout->SetPlane(CylinderOptions_DEPRECATED.Plane);
-			CylinderLayout->SetRotation(CylinderOptions_DEPRECATED.Rotation);
-			CylinderLayout->SetScale(CylinderOptions_DEPRECATED.Scale);
-		}
-
-		// Sphere options
-		if (UCEClonerSphereUniformLayout* SphereLayout = FindOrAddLayout<UCEClonerSphereUniformLayout>())
-		{
-			SphereLayout->SetCount(SphereOptions_DEPRECATED.Count);
-			SphereLayout->SetRadius(SphereOptions_DEPRECATED.Radius);
-			SphereLayout->SetRatio(SphereOptions_DEPRECATED.Ratio);
-			SphereLayout->SetOrientMesh(SphereOptions_DEPRECATED.bOrientMesh);
-			SphereLayout->SetRotation(SphereOptions_DEPRECATED.Rotation);
-			SphereLayout->SetScale(SphereOptions_DEPRECATED.Scale);
-		}
-
-		// Honeycomb options
-		if (UCEClonerHoneycombLayout* HoneycombLayout = FindOrAddLayout<UCEClonerHoneycombLayout>())
-		{
-			HoneycombLayout->SetPlane(HoneycombOptions_DEPRECATED.Plane);
-			HoneycombLayout->SetWidthCount(HoneycombOptions_DEPRECATED.WidthCount);
-			HoneycombLayout->SetHeightCount(HoneycombOptions_DEPRECATED.HeightCount);
-			HoneycombLayout->SetWidthOffset(HoneycombOptions_DEPRECATED.WidthOffset);
-			HoneycombLayout->SetHeightOffset(HoneycombOptions_DEPRECATED.HeightOffset);
-			HoneycombLayout->SetHeightSpacing(HoneycombOptions_DEPRECATED.HeightSpacing);
-			HoneycombLayout->SetWidthSpacing(HoneycombOptions_DEPRECATED.WidthSpacing);
-		}
-
-		// Sample mesh options
-		if (UCEClonerMeshLayout* MeshLayout = FindOrAddLayout<UCEClonerMeshLayout>())
-		{
-			MeshLayout->SetCount(SampleMeshOptions_DEPRECATED.Count);
-			MeshLayout->SetAsset(SampleMeshOptions_DEPRECATED.Asset);
-			MeshLayout->SetSampleData(SampleMeshOptions_DEPRECATED.SampleData);
-			MeshLayout->SetSampleActorWeak(SampleMeshOptions_DEPRECATED.SampleActor);
-		}
-
-		// Sample spline options
-		if (UCEClonerSplineLayout* SplineLayout = FindOrAddLayout<UCEClonerSplineLayout>())
-		{
-			SplineLayout->SetCount(SampleSplineOptions_DEPRECATED.Count);
-			SplineLayout->SetSplineActorWeak(SampleSplineOptions_DEPRECATED.SplineActor);
-		}
-
-		// Set previous layout
-		if (const UCEClonerSubsystem* Subsystem = UCEClonerSubsystem::Get())
-		{
-			const int32 LayoutIdx = static_cast<int32>(Layout_DEPRECATED);
-			SetLayoutName(Subsystem->GetLayoutNames()[LayoutIdx]);
-		}
-
-		bDeprecatedPropertiesMigrated = true;
-
-		UE_LOG(LogCEClonerActor, Log, TEXT("%s : Cloner deprecated properties migrated"), *GetActorNameOrLabel());
-	}
-
 	for (TMap<FName, TObjectPtr<UCEClonerLayoutBase>>::TIterator It(LayoutInstances); It; ++It)
 	{
 		if (!It->Value)
@@ -1525,9 +1410,6 @@ void ACEClonerActor::InitializeCloner()
 	}
 #endif
 
-	// For new cloner instances no need to migrate anything
-	bDeprecatedPropertiesMigrated = true;
-
 	bClonerInitialized = true;
 
 	OnLayoutNameChanged();
diff --git a/Engine/Plugins/Experimental/ClonerEffector/Source/ClonerEffector/Private/Effector/CEEffectorActor.cpp b/Engine/Plugins/Experimental/ClonerEffector/Source/ClonerEffector/Private/Effector/CEEffectorActor.cpp
--- a/Engine/Plugins/Experimental/ClonerEffector/Source/ClonerEffector/Private/Effector/CEEffectorActor.cpp
+++ b/Engine/Plugins/Experimental/ClonerEffector/Source/ClonerEffector/Private/Effector/CEEffectorActor.cpp
@@ -183,19 +183,6 @@ void ACEEffectorActor::PostLoad()
 	Super::PostLoad();
 
 	OnEffectorChanged();
-
-	if (!InternalCloners.IsEmpty())
-	{
-		for (const TWeakObjectPtr<ACEClonerActor>& ClonerWeak : InternalCloners)
-		{
-			if (ACEClonerActor* Cloner = ClonerWeak.Get())
-			{
-				Cloner->LinkEffector(this);
-			}
-		}
-
-		InternalCloners.Empty();
-	}
 }
 
 void ACEEffectorActor::PostDuplicate(EDuplicateMode::Type InDuplicateMode)
diff --git a/Engine/Plugins/Experimental/ClonerEffector/Source/ClonerEffector/Public/CEClonerEffectorShared.h b/Engine/Plugins/Experimental/ClonerEffector/Source/ClonerEffector/Public/CEClonerEffectorShared.h
--- a/Engine/Plugins/Experimental/ClonerEffector/Source/ClonerEffector/Public/CEClonerEffectorShared.h
+++ b/Engine/Plugins/Experimental/ClonerEffector/Source/ClonerEffector/Public/CEClonerEffectorShared.h
@@ -15,20 +15,6 @@ class UDynamicMesh;
 class UNiagaraSystem;
 class USplineComponent;
 
-/** Enumerate the different layouts available to dispose cloner instances */
-UENUM(BlueprintType)
-enum class ECEClonerLayout : uint8
-{
-	Grid,
-	Line,
-	Circle,
-	Cylinder,
-	Sphere,
-	Honeycomb,
-	SampleMesh,
-	SampleSpline
-};
-
 /** Enumerate the axis available to use */
 UENUM(BlueprintType)
 enum class ECEClonerAxis : uint8
@@ -420,42 +406,6 @@ protected:
 	TMap<FName, TObjectPtr<UNiagaraDataInterface>> DataInterfaces;
 };
 
-USTRUCT(BlueprintType)
-struct CLONEREFFECTOR_API FCEClonerSampleMeshOptions
-{
-	GENERATED_BODY()
-
-	UPROPERTY(BlueprintReadWrite, EditInstanceOnly, Category="Cloner")
-	int32 Count = 3 * 3 * 3;
-
-	UPROPERTY(BlueprintReadWrite, EditInstanceOnly, Category="Cloner")
-	ECEClonerMeshAsset Asset = ECEClonerMeshAsset::StaticMesh;
-
-	UPROPERTY(BlueprintReadWrite, EditInstanceOnly, Category="Cloner")
-	ECEClonerMeshSampleData SampleData = ECEClonerMeshSampleData::Vertices;
-
-	UPROPERTY(BlueprintReadWrite, EditInstanceOnly, Category="Cloner")
-	TWeakObjectPtr<AActor> SampleActor = nullptr;
-
-	UPROPERTY()
-	TWeakObjectPtr<USceneComponent> SceneComponent = nullptr;
-};
-
-USTRUCT(BlueprintType)
-struct CLONEREFFECTOR_API FCEClonerSampleSplineOptions
-{
-	GENERATED_BODY()
-
-	UPROPERTY(BlueprintReadWrite, EditInstanceOnly, Category="Cloner")
-	int32 Count = 3 * 3;
-
-	UPROPERTY(BlueprintReadWrite, EditInstanceOnly, Category="Cloner")
-	TWeakObjectPtr<AActor> SplineActor = nullptr;
-
-	UPROPERTY()
-	TWeakObjectPtr<USplineComponent> SplineComponent = nullptr;
-};
-
 USTRUCT(BlueprintType)
 struct CLONEREFFECTOR_API FCEClonerGridConstraintSphere
 {
@@ -504,171 +454,6 @@ struct CLONEREFFECTOR_API FCEClonerGridConstraintTexture
 	float Threshold = 0.f;
 };
 
-USTRUCT(BlueprintType)
-struct CLONEREFFECTOR_API FCEClonerGridLayoutOptions
-{
-	GENERATED_BODY()
-
-	UPROPERTY(BlueprintReadWrite, EditInstanceOnly, Category="Cloner")
-	FIntVector Count = FIntVector(3, 3, 3);
-
-	UPROPERTY(BlueprintReadWrite, EditInstanceOnly, Category="Cloner")
-	FVector Spacing = FVector(400.f, 400.f, 400.f);
-
-	UPROPERTY(BlueprintReadWrite, EditInstanceOnly, Category="Cloner")
-	ECEClonerGridConstraint Constraint = ECEClonerGridConstraint::None;
-
-	UPROPERTY(BlueprintReadWrite, EditInstanceOnly, Category="Cloner", meta=(EditCondition="Constraint != ECEClonerGridConstraint::None", EditConditionHides))
-	bool bInvertConstraint = false;
-
-	UPROPERTY(BlueprintReadWrite, EditInstanceOnly, Category="Cloner", meta=(EditCondition="Constraint == ECEClonerGridConstraint::Sphere", EditConditionHides))
-	FCEClonerGridConstraintSphere SphereConstraint;
-
-	UPROPERTY(BlueprintReadWrite, EditInstanceOnly, Category="Cloner", meta=(EditCondition="Constraint == ECEClonerGridConstraint::Cylinder", EditConditionHides))
-	FCEClonerGridConstraintCylinder CylinderConstraint;
-
-	UPROPERTY(BlueprintReadWrite, EditInstanceOnly, Category="Cloner", meta=(EditCondition="Constraint == ECEClonerGridConstraint::Texture", EditConditionHides))
-	FCEClonerGridConstraintTexture TextureConstraint;
-};
-
-USTRUCT(BlueprintType)
-struct CLONEREFFECTOR_API FCEClonerLineLayoutOptions
-{
-	GENERATED_BODY()
-
-	UPROPERTY(BlueprintReadWrite, EditInstanceOnly, Category="Cloner")
-	int32 Count = 3 * 3;
-
-	UPROPERTY(BlueprintReadWrite, EditInstanceOnly, Category="Cloner")
-	float Spacing = 400.f;
-
-	UPROPERTY(BlueprintReadWrite, EditInstanceOnly, Category="Cloner")
-	ECEClonerAxis Axis = ECEClonerAxis::Y;
-
-	UPROPERTY(BlueprintReadWrite, EditInstanceOnly, Category="Cloner", meta=(ClampMin="0", ClampMax="1", EditCondition="Axis == ECEClonerAxis::Custom", EditConditionHides))
-	FVector Direction = FVector::YAxisVector;
-
-	UPROPERTY(BlueprintReadWrite, EditInstanceOnly, Category="Cloner")
-	FRotator Rotation = FRotator(0.f);
-};
-
-USTRUCT(BlueprintType)
-struct CLONEREFFECTOR_API FCEClonerCircleLayoutOptions
-{
-	GENERATED_BODY()
-
-	UPROPERTY(BlueprintReadWrite, EditInstanceOnly, Category="Cloner")
-	int32 Count = 3 * 3;
-
-	UPROPERTY(BlueprintReadWrite, EditInstanceOnly, Category="Cloner")
-	float Radius = 400.f;
-
-	UPROPERTY(BlueprintReadWrite, EditInstanceOnly, Category="Cloner")
-	float AngleStart = 0.f;
-
-	UPROPERTY(BlueprintReadWrite, EditInstanceOnly, Category="Cloner")
-	float AngleRatio = 1.f;
-
-	UPROPERTY(BlueprintReadWrite, EditInstanceOnly, Category="Cloner")
-	bool bOrientMesh = false;
-
-	UPROPERTY(BlueprintReadWrite, EditInstanceOnly, Category="Cloner")
-	ECEClonerPlane Plane = ECEClonerPlane::XY;
-
-	UPROPERTY(BlueprintReadWrite, EditInstanceOnly, Category="Cloner", meta=(EditCondition="Plane == ECEClonerPlane::Custom", EditConditionHides))
-	FRotator Rotation = FRotator(0.f, 0.f, 0.f);
-
-	UPROPERTY(BlueprintReadWrite, EditInstanceOnly, Category="Cloner", meta=(ClampMin="0", AllowPreserveRatio, Delta="0.01"))
-	FVector Scale = FVector(1.f, 1.f, 1.f);
-};
-
-USTRUCT(BlueprintType)
-struct CLONEREFFECTOR_API FCEClonerCylinderLayoutOptions
-{
-	GENERATED_BODY()
-
-	UPROPERTY(BlueprintReadWrite, EditInstanceOnly, Category="Cloner")
-	int32 BaseCount = 3 * 3;
-
-	UPROPERTY(BlueprintReadWrite, EditInstanceOnly, Category="Cloner")
-	int32 HeightCount = 3;
-
-	UPROPERTY(BlueprintReadWrite, EditInstanceOnly, Category="Cloner")
-	float Height = 400.f;
-
-	UPROPERTY(BlueprintReadWrite, EditInstanceOnly, Category="Cloner")
-	float Radius = 400.f;
-
-	UPROPERTY(BlueprintReadWrite, EditInstanceOnly, Category="Cloner")
-	float AngleStart = 0.f;
-
-	UPROPERTY(BlueprintReadWrite, EditInstanceOnly, Category="Cloner")
-	float AngleRatio = 1.f;
-
-	UPROPERTY(BlueprintReadWrite, EditInstanceOnly, Category="Cloner")
-	bool bOrientMesh = false;
-
-	UPROPERTY(BlueprintReadWrite, EditInstanceOnly, Category="Cloner")
-	ECEClonerPlane Plane = ECEClonerPlane::XY;
-
-	UPROPERTY(BlueprintReadWrite, EditInstanceOnly, Category="Cloner", meta=(EditCondition="Plane == ECEClonerPlane::Custom", EditConditionHides))
-	FRotator Rotation = FRotator(0.f, 0.f, 0.f);
-
-	UPROPERTY(BlueprintReadWrite, EditInstanceOnly, Category="Cloner", meta=(ClampMin="0", AllowPreserveRatio, Delta="0.01"))
-	FVector Scale = FVector(1.f, 1.f, 1.f);
-};
-
-USTRUCT(BlueprintType)
-struct CLONEREFFECTOR_API FCEClonerSphereLayoutOptions
-{
-	GENERATED_BODY()
-
-	UPROPERTY(BlueprintReadWrite, EditInstanceOnly, Category="Cloner")
-	int32 Count = 3 * 3 * 3;
-
-	UPROPERTY(BlueprintReadWrite, EditInstanceOnly, Category="Cloner")
-	float Radius = 400.f;
-
-	UPROPERTY(BlueprintReadWrite, EditInstanceOnly, Category="Cloner", meta=(ClampMin="0", ClampMax="1"))
-	float Ratio = 1.f;
-
-	UPROPERTY(BlueprintReadWrite, EditInstanceOnly, Category="Cloner")
-	bool bOrientMesh = false;
-
-	UPROPERTY(BlueprintReadWrite, EditInstanceOnly, Category="Cloner")
-	FRotator Rotation = FRotator(0.f, 0.f, 0.f);
-
-	UPROPERTY(BlueprintReadWrite, EditInstanceOnly, Category="Cloner", meta=(ClampMin="0", AllowPreserveRatio, Delta="0.01"))
-	FVector Scale = FVector(1.f, 1.f, 1.f);
-};
-
-USTRUCT(BlueprintType)
-struct CLONEREFFECTOR_API FCEClonerHoneycombLayoutOptions
-{
-	GENERATED_BODY()
-
-	UPROPERTY(BlueprintReadWrite, EditInstanceOnly, Category="Cloner")
-	ECEClonerPlane Plane = ECEClonerPlane::XY;
-
-	UPROPERTY(BlueprintReadWrite, EditInstanceOnly, Category="Cloner")
-	int32 WidthCount = 3;
-
-	UPROPERTY(BlueprintReadWrite, EditInstanceOnly, Category="Cloner")
-	int32 HeightCount = 3;
-
-	UPROPERTY(BlueprintReadWrite, EditInstanceOnly, Category="Cloner")
-	float WidthOffset = 0.f;
-
-	UPROPERTY(BlueprintReadWrite, EditInstanceOnly, Category="Cloner")
-	float HeightOffset = 0.5f;
-
-	UPROPERTY(BlueprintReadWrite, EditInstanceOnly, Category="Cloner")
-	float WidthSpacing = 400.f;
-
-	UPROPERTY(BlueprintReadWrite, EditInstanceOnly, Category="Cloner")
-	float HeightSpacing = 400.f;
-};
-
 namespace UE::ClonerEffector
 {
 	void SetBillboardComponentSprite(const AActor* InActor, const FString& InTexturePath);
diff --git a/Engine/Plugins/Experimental/ClonerEffector/Source/ClonerEffector/Public/Cloner/CEClonerActor.h b/Engine/Plugins/Experimental/ClonerEffector/Source/ClonerEffector/Public/Cloner/CEClonerActor.h
--- a/Engine/Plugins/Experimental/ClonerEffector/Source/ClonerEffector/Public/Cloner/CEClonerActor.h
+++ b/Engine/Plugins/Experimental/ClonerEffector/Source/ClonerEffector/Public/Cloner/CEClonerActor.h
@@ -698,45 +698,6 @@ private:
 	UPROPERTY(EditInstanceOnly, Category="Cloner", meta=(DisplayName="Effectors"))
 	TArray<TWeakObjectPtr<ACEEffectorActor>> EffectorsWeak;
 
-	/**
-	* Below properties are deprecated and no longer in use,
-	* they will be migrated to new layout system on load,
-	* And they will be all removed in future version
-	*/
-	UPROPERTY(meta=(DeprecatedProperty, DeprecationMessage="Use LayoutName instead"))
-	ECEClonerLayout Layout_DEPRECATED = ECEClonerLayout::Grid;
-
-	UPROPERTY(meta=(DeprecatedProperty, DeprecationMessage="Use UCEClonerGridLayout instead"))
-	FCEClonerGridLayoutOptions GridOptions_DEPRECATED;
-
-	UPROPERTY(meta=(DeprecatedProperty, DeprecationMessage="Use UCEClonerLineLayout instead"))
-	FCEClonerLineLayoutOptions LineOptions_DEPRECATED;
-
-	UPROPERTY(meta=(DeprecatedProperty, DeprecationMessage="Use UCEClonerCircleLayout instead"))
-	FCEClonerCircleLayoutOptions CircleOptions_DEPRECATED;
-
-	UPROPERTY(meta=(DeprecatedProperty, DeprecationMessage="Use UCEClonerCylinderLayout instead"))
-	FCEClonerCylinderLayoutOptions CylinderOptions_DEPRECATED;
-
-	UPROPERTY(meta=(DeprecatedProperty, DeprecationMessage="Use UCEClonerSphereLayout instead"))
-	FCEClonerSphereLayoutOptions SphereOptions_DEPRECATED;
-
-	UPROPERTY(meta=(DeprecatedProperty, DeprecationMessage="Use UCEClonerHoneycombLayout instead"))
-	FCEClonerHoneycombLayoutOptions HoneycombOptions_DEPRECATED;
-
-	UPROPERTY(meta=(DeprecatedProperty, DeprecationMessage="Use UCEClonerMeshLayout instead"))
-	FCEClonerSampleMeshOptions SampleMeshOptions_DEPRECATED;
-
-	UPROPERTY(meta=(DeprecatedProperty, DeprecationMessage="Use UCEClonerSplineLayout instead"))
-	FCEClonerSampleSplineOptions SampleSplineOptions_DEPRECATED;
-
-	/**
-	 * Used to confirm migration of deprecated properties above to new layout system
-	 * Newly created cloner will not migrate anything and skip this task
-	 */
-	UPROPERTY()
-	bool bDeprecatedPropertiesMigrated = false;
-
 	float TreeUpdateDeltaTime = 0.f;
 
 	bool bNeedsRefresh = false;
diff --git a/Engine/Plugins/Experimental/ClonerEffector/Source/ClonerEffector/Public/Effector/CEEffectorActor.h b/Engine/Plugins/Experimental/ClonerEffector/Source/ClonerEffector/Public/Effector/CEEffectorActor.h
--- a/Engine/Plugins/Experimental/ClonerEffector/Source/ClonerEffector/Public/Effector/CEEffectorActor.h
+++ b/Engine/Plugins/Experimental/ClonerEffector/Source/ClonerEffector/Public/Effector/CEEffectorActor.h
@@ -742,10 +742,6 @@ private:
 	FInstancedPropertyBag VisualizerData;
 #endif
 
-	/** Internal cloners array, deprecated since it will be moved to the cloners and emptied out on post load */
-	UPROPERTY(NonTransactional)
-	TSet<TWeakObjectPtr<ACEClonerActor>> InternalCloners;
-
 	/** Transient effector channel data */
 	UPROPERTY(VisibleInstanceOnly, Transient, DuplicateTransient, TextExportTransient, NonTransactional, AdvancedDisplay, Category="Effector")
 	FCEClonerEffectorChannelData ChannelData;
