MotionDesign : ClonerEffector
- Duplicating an effector will relink it to cloners it was previously linked to

diff --git a/Engine/Plugins/Experimental/ClonerEffector/Source/ClonerEffector/Private/Cloner/CEClonerActor.cpp b/Engine/Plugins/Experimental/ClonerEffector/Source/ClonerEffector/Private/Cloner/CEClonerActor.cpp
--- a/Engine/Plugins/Experimental/ClonerEffector/Source/ClonerEffector/Private/Cloner/CEClonerActor.cpp
+++ b/Engine/Plugins/Experimental/ClonerEffector/Source/ClonerEffector/Private/Cloner/CEClonerActor.cpp
@@ -56,7 +56,6 @@ ACEClonerActor::ACEClonerActor()
 		ClonerComponent->OnClonerMeshUpdated.AddUObject(this, &ACEClonerActor::OnClonerMeshUpdated);
 
 		UCEEffectorSubsystem::OnEffectorIdentifierChangedDelegate.AddUObject(this, &ACEClonerActor::OnEffectorIdentifierChanged);
-		ACEEffectorActor::OnEffectorRefreshClonerDelegate.AddUObject(this, &ACEClonerActor::OnEffectorRefreshCloner);
 
 		const TArray<FString> LayoutNames = GetClonerLayoutNames();
 
@@ -1022,14 +1021,6 @@ void ACEClonerActor::OnEffectorIdentifierChanged(ACEEffectorActor* InEffector, i
 	}
 }
 
-void ACEClonerActor::OnEffectorRefreshCloner(ACEEffectorActor* InEffector)
-{
-	if (EffectorsWeak.Contains(InEffector))
-	{
-		RequestClonerUpdate();
-	}
-}
-
 void ACEClonerActor::OnEffectorsChanged()
 {
 	const FCEClonerEffectorDataInterfaces* EffectorDataInterfaces = GetEffectorDataInterfaces();
@@ -1043,6 +1034,26 @@ void ACEClonerActor::OnEffectorsChanged()
 	const TSet<TWeakObjectPtr<ACEEffectorActor>> SetEffectorsWeak(EffectorsWeak);
 	EffectorsWeak = SetEffectorsWeak.Array();
 
+	// Removed effectors
+	for (const TWeakObjectPtr<ACEEffectorActor>& EffectorWeak : EffectorsInternalWeak.Difference(SetEffectorsWeak))
+	{
+		if (ACEEffectorActor* Effector = EffectorWeak.Get())
+		{
+			Effector->OnClonerUnlinked(this);
+		}
+	}
+
+	// Added effectors
+	for (const TWeakObjectPtr<ACEEffectorActor>& EffectorWeak : SetEffectorsWeak.Difference(EffectorsInternalWeak))
+	{
+		if (ACEEffectorActor* Effector = EffectorWeak.Get())
+		{
+			Effector->OnClonerLinked(this);
+		}
+	}
+
+	EffectorsInternalWeak = SetEffectorsWeak;
+
 	TArray<int32> EffectorIndexes;
 	EffectorIndexes.Reserve(EffectorsWeak.Num());
 
diff --git a/Engine/Plugins/Experimental/ClonerEffector/Source/ClonerEffector/Private/Effector/CEEffectorActor.cpp b/Engine/Plugins/Experimental/ClonerEffector/Source/ClonerEffector/Private/Effector/CEEffectorActor.cpp
--- a/Engine/Plugins/Experimental/ClonerEffector/Source/ClonerEffector/Private/Effector/CEEffectorActor.cpp
+++ b/Engine/Plugins/Experimental/ClonerEffector/Source/ClonerEffector/Private/Effector/CEEffectorActor.cpp
@@ -16,7 +16,7 @@
 #include "UObject/ConstructorHelpers.h"
 #endif
 
-ACEEffectorActor::FOnEffectorIdentifierChanged ACEEffectorActor::OnEffectorRefreshClonerDelegate;
+DEFINE_LOG_CATEGORY_STATIC(LogCEEffectorActor, Log, All);
 
 ACEEffectorActor::ACEEffectorActor()
 {
@@ -189,6 +189,15 @@ void ACEEffectorActor::PostDuplicate(EDuplicateMode::Type InDuplicateMode)
 {
 	Super::PostDuplicate(InDuplicateMode);
 
+	for (const TWeakObjectPtr<ACEClonerActor>& ClonerWeak : ClonersWeak)
+	{
+		if (ACEClonerActor* Cloner = ClonerWeak.Get())
+		{
+			UE_LOG(LogCEEffectorActor, Log, TEXT("Linking effector %s to cloner %s after duplication"), *GetActorNameOrLabel(), *Cloner->GetActorNameOrLabel())
+			Cloner->LinkEffector(this);
+		}
+	}
+
 	OnEffectorChanged();
 }
 
@@ -217,6 +226,26 @@ void ACEEffectorActor::OnEffectorSubsystemInitialized(const UWorld* InWorld)
 	}
 }
 
+void ACEEffectorActor::OnClonerLinked(ACEClonerActor* InCloner)
+{
+	if (!IsValid(InCloner) || ClonersWeak.Contains(InCloner))
+	{
+		return;
+	}
+
+	ClonersWeak.Add(InCloner);
+}
+
+void ACEEffectorActor::OnClonerUnlinked(ACEClonerActor* InCloner)
+{
+	if (!IsValid(InCloner) || !ClonersWeak.Contains(InCloner))
+	{
+		return;
+	}
+
+	ClonersWeak.Remove(InCloner);
+}
+
 FCEClonerEffectorChannelData& ACEEffectorActor::GetChannelData()
 {
 	return ChannelData;
@@ -1384,7 +1413,13 @@ void ACEEffectorActor::OnForceEnabledChanged()
 	if (bEnabled)
 	{
 		// Refresh cloners to reset clones transform
-		OnEffectorRefreshClonerDelegate.Broadcast(this);
+		for (const TWeakObjectPtr<ACEClonerActor>& ClonerWeak : ClonersWeak)
+		{
+			if (ACEClonerActor* Cloner = ClonerWeak.Get())
+			{
+				Cloner->RequestClonerUpdate();
+			}
+		}
 	}
 
 	OnForceOptionsChanged();
diff --git a/Engine/Plugins/Experimental/ClonerEffector/Source/ClonerEffector/Private/Subsystems/CEEffectorSubsystem.cpp b/Engine/Plugins/Experimental/ClonerEffector/Source/ClonerEffector/Private/Subsystems/CEEffectorSubsystem.cpp
--- a/Engine/Plugins/Experimental/ClonerEffector/Source/ClonerEffector/Private/Subsystems/CEEffectorSubsystem.cpp
+++ b/Engine/Plugins/Experimental/ClonerEffector/Source/ClonerEffector/Private/Subsystems/CEEffectorSubsystem.cpp
@@ -88,6 +88,22 @@ bool UCEEffectorSubsystem::UnregisterChannelEffector(ACEEffectorActor* InEffecto
 	return bUnregistered;
 }
 
+ACEEffectorActor* UCEEffectorSubsystem::GetEffectorByChannelIdentifier(int32 InIdentifier) const
+{
+	if (EffectorsWeak.IsValidIndex(InIdentifier))
+	{
+		if (ACEEffectorActor* Effector = EffectorsWeak[InIdentifier].Get())
+		{
+			if (Effector->GetChannelIdentifier() == InIdentifier)
+			{
+				return Effector;
+			}
+		}
+	}
+
+	return nullptr;
+}
+
 void UCEEffectorSubsystem::UpdateEffectorChannel()
 {
 	const UWorld* World = GetWorld();
diff --git a/Engine/Plugins/Experimental/ClonerEffector/Source/ClonerEffector/Private/Subsystems/CEEffectorSubsystem.h b/Engine/Plugins/Experimental/ClonerEffector/Source/ClonerEffector/Private/Subsystems/CEEffectorSubsystem.h
--- a/Engine/Plugins/Experimental/ClonerEffector/Source/ClonerEffector/Private/Subsystems/CEEffectorSubsystem.h
+++ b/Engine/Plugins/Experimental/ClonerEffector/Source/ClonerEffector/Private/Subsystems/CEEffectorSubsystem.h
@@ -39,6 +39,8 @@ public:
 	/** Unregister an effector actor used within a effector channel */
 	bool UnregisterChannelEffector(ACEEffectorActor* InEffector);
 
+	ACEEffectorActor* GetEffectorByChannelIdentifier(int32 InIdentifier) const;
+
 protected:
 	//~ Begin FTickableGameObject interface
 	virtual bool IsTickableInEditor() const override;
diff --git a/Engine/Plugins/Experimental/ClonerEffector/Source/ClonerEffector/Public/Cloner/CEClonerActor.h b/Engine/Plugins/Experimental/ClonerEffector/Source/ClonerEffector/Public/Cloner/CEClonerActor.h
--- a/Engine/Plugins/Experimental/ClonerEffector/Source/ClonerEffector/Public/Cloner/CEClonerActor.h
+++ b/Engine/Plugins/Experimental/ClonerEffector/Source/ClonerEffector/Public/Cloner/CEClonerActor.h
@@ -473,7 +473,6 @@ protected:
 	void OnClonerSystemChanged();
 
 	void OnEffectorIdentifierChanged(ACEEffectorActor* InEffector, int32 InOldIdentifier, int32 InNewIdentifier);
-	void OnEffectorRefreshCloner(ACEEffectorActor* InEffector);
 
 	void OnEffectorsChanged();
 	void OnEnabledChanged();
@@ -698,6 +697,10 @@ private:
 	UPROPERTY(EditInstanceOnly, Category="Cloner", meta=(DisplayName="Effectors"))
 	TArray<TWeakObjectPtr<ACEEffectorActor>> EffectorsWeak;
 
+	/** Copy of effectors linked to this cloner to compare diffs */
+	UPROPERTY(Transient, DuplicateTransient, TextExportTransient, NonTransactional)
+	TSet<TWeakObjectPtr<ACEEffectorActor>> EffectorsInternalWeak;
+
 	float TreeUpdateDeltaTime = 0.f;
 
 	bool bNeedsRefresh = false;
diff --git a/Engine/Plugins/Experimental/ClonerEffector/Source/ClonerEffector/Public/Effector/CEEffectorActor.h b/Engine/Plugins/Experimental/ClonerEffector/Source/ClonerEffector/Public/Effector/CEEffectorActor.h
--- a/Engine/Plugins/Experimental/ClonerEffector/Source/ClonerEffector/Public/Effector/CEEffectorActor.h
+++ b/Engine/Plugins/Experimental/ClonerEffector/Source/ClonerEffector/Public/Effector/CEEffectorActor.h
@@ -432,10 +432,6 @@ public:
 #endif
 
 protected:
-	/** Used to trigger a refresh on linked cloner */
-	DECLARE_MULTICAST_DELEGATE_OneParam(FOnEffectorIdentifierChanged, ACEEffectorActor* /** InEffectorActor */)
-	static FOnEffectorIdentifierChanged OnEffectorRefreshClonerDelegate;
-
 	static constexpr int32 InnerVisualizerId = 0;
 	static constexpr int32 OuterVisualizerId = 1;
 	static constexpr TCHAR VisualizerColorName[] = TEXT("VisualizerColor");
@@ -462,6 +458,9 @@ protected:
 	int32 GetChannelIdentifier() const;
 	void OnEffectorSubsystemInitialized(const UWorld* InWorld);
 
+	void OnClonerLinked(ACEClonerActor* InCloner);
+	void OnClonerUnlinked(ACEClonerActor* InCloner);
+
 	FCEClonerEffectorChannelData& GetChannelData();
 
 	void OnEffectorTransformed(USceneComponent* InUpdatedComponent, EUpdateTransformFlags InUpdateTransformFlags, ETeleportType InTeleport);
@@ -746,6 +745,10 @@ private:
 	UPROPERTY(VisibleInstanceOnly, Transient, DuplicateTransient, TextExportTransient, NonTransactional, AdvancedDisplay, Category="Effector")
 	FCEClonerEffectorChannelData ChannelData;
 
+	/** Cloners linked to this effector, used to refresh or relink on duplicate */
+	UPROPERTY(SkipSerialization, NonTransactional)
+	TArray<TWeakObjectPtr<ACEClonerActor>> ClonersWeak;
+
 #if WITH_EDITOR
 	/** Used for PECP */
 	static TCEPropertyChangeDispatcher<ACEEffectorActor> PropertyChangeDispatcher;
