MotionDesign : ClonerEffector
- Add effector push mode support with specific directions (forward, right, up, effector, random)
- Add cloner surface and particle collisions support using neighbor grid

#virtualized

diff --git a/Engine/Plugins/Experimental/ClonerEffector/Source/ClonerEffector/Private/Cloner/CEClonerActor.cpp b/Engine/Plugins/Experimental/ClonerEffector/Source/ClonerEffector/Private/Cloner/CEClonerActor.cpp
--- a/Engine/Plugins/Experimental/ClonerEffector/Source/ClonerEffector/Private/Cloner/CEClonerActor.cpp
+++ b/Engine/Plugins/Experimental/ClonerEffector/Source/ClonerEffector/Private/Cloner/CEClonerActor.cpp
@@ -7,6 +7,7 @@
 #include "Effector/CEEffectorActor.h"
 #include "Engine/StaticMeshActor.h"
 #include "Materials/MaterialInterface.h"
+#include "NiagaraDataInterfaceArrayFloat.h"
 #include "NiagaraDataInterfaceCurve.h"
 #include "NiagaraDataInterfaceSkeletalMesh.h"
 #include "NiagaraMeshRendererProperties.h"
@@ -92,6 +93,17 @@ const TCEPropertyChangeDispatcher<ACEClonerActor> ACEClonerActor::PropertyChange
 	{ GET_MEMBER_NAME_CHECKED(ACEClonerActor, Color), &ACEClonerActor::OnColorChanged },
 	{ GET_MEMBER_NAME_CHECKED(ACEClonerActor, bVisualizeEffectors), &ACEClonerActor::OnOverrideMaterialChanged },
 	{ GET_MEMBER_NAME_CHECKED(ACEClonerActor, EffectorsWeak), &ACEClonerActor::OnEffectorsChanged },
+	/** Collision */
+	{ GET_MEMBER_NAME_CHECKED(ACEClonerActor, bSurfaceCollisionEnabled), &ACEClonerActor::OnCollisionOptionsChanged },
+	{ GET_MEMBER_NAME_CHECKED(ACEClonerActor, bParticleCollisionEnabled), &ACEClonerActor::OnCollisionOptionsChanged },
+	{ GET_MEMBER_NAME_CHECKED(ACEClonerActor, bCollisionVelocityEnabled), &ACEClonerActor::OnCollisionOptionsChanged },
+	{ GET_MEMBER_NAME_CHECKED(ACEClonerActor, CollisionRadiusMode), &ACEClonerActor::OnCollisionOptionsChanged },
+	{ GET_MEMBER_NAME_CHECKED(ACEClonerActor, CollisionRadii), &ACEClonerActor::OnCollisionOptionsChanged },
+	{ GET_MEMBER_NAME_CHECKED(ACEClonerActor, CollisionIterations), &ACEClonerActor::OnCollisionOptionsChanged },
+	{ GET_MEMBER_NAME_CHECKED(ACEClonerActor, CollisionGridResolution), &ACEClonerActor::OnCollisionOptionsChanged },
+	{ GET_MEMBER_NAME_CHECKED(ACEClonerActor, CollisionGridSize), &ACEClonerActor::OnCollisionOptionsChanged },
+	{ GET_MEMBER_NAME_CHECKED(ACEClonerActor, MassMin), &ACEClonerActor::OnCollisionOptionsChanged },
+	{ GET_MEMBER_NAME_CHECKED(ACEClonerActor, MassMax), &ACEClonerActor::OnCollisionOptionsChanged },
 	/** Layout */
 	{ GET_MEMBER_NAME_CHECKED(ACEClonerActor, LayoutName), &ACEClonerActor::OnLayoutNameChanged },
 	/** Advanced */
@@ -168,6 +180,7 @@ void ACEClonerActor::PostInitProperties()
 		FNiagaraTypeRegistry::Register(FNiagaraTypeDefinition(StaticEnum<ECEClonerEffectorMode>()), MeshFlags);
 		FNiagaraTypeRegistry::Register(FNiagaraTypeDefinition(StaticEnum<ECEClonerSpawnLoopMode>()), MeshFlags);
 		FNiagaraTypeRegistry::Register(FNiagaraTypeDefinition(StaticEnum<ECEClonerSpawnBehaviorMode>()), MeshFlags);
+		FNiagaraTypeRegistry::Register(FNiagaraTypeDefinition(StaticEnum<ECEClonerEffectorPushDirection>()), MeshFlags);
 	}
 }
 
@@ -258,6 +271,7 @@ void ACEClonerActor::UpdateLayoutOptions()
 	OnRangeOptionsChanged();
 	OnSpawnOptionsChanged();
 	OnLifetimeOptionsChanged();
+	OnCollisionOptionsChanged();
 }
 
 void ACEClonerActor::SetTreeUpdateInterval(float InInterval)
@@ -385,6 +399,110 @@ void ACEClonerActor::SetVisualizeEffectors(bool bInVisualize)
 	OnOverrideMaterialChanged();
 }
 
+void ACEClonerActor::SetSurfaceCollisionEnabled(bool bInSurfaceCollisionEnabled)
+{
+	if (bSurfaceCollisionEnabled == bInSurfaceCollisionEnabled)
+	{
+		return;
+	}
+
+	bSurfaceCollisionEnabled = bInSurfaceCollisionEnabled;
+	OnCollisionOptionsChanged();
+}
+
+void ACEClonerActor::SetParticleCollisionEnabled(bool bInParticleCollisionEnabled)
+{
+	if (bParticleCollisionEnabled == bInParticleCollisionEnabled)
+	{
+		return;
+	}
+
+	bParticleCollisionEnabled = bInParticleCollisionEnabled;
+	OnCollisionOptionsChanged();
+}
+
+void ACEClonerActor::SetCollisionVelocityEnabled(bool bInCollisionVelocityEnabled)
+{
+	if (bCollisionVelocityEnabled == bInCollisionVelocityEnabled)
+	{
+		return;
+	}
+
+	bCollisionVelocityEnabled = bInCollisionVelocityEnabled;
+	OnCollisionOptionsChanged();
+}
+
+void ACEClonerActor::SetCollisionIterations(int32 InCollisionIterations)
+{
+	InCollisionIterations = FMath::Max(InCollisionIterations, 1);
+	if (CollisionIterations == InCollisionIterations)
+	{
+		return;
+	}
+
+	CollisionIterations = InCollisionIterations;
+	OnCollisionOptionsChanged();
+}
+
+void ACEClonerActor::SetCollisionGridResolution(int32 InCollisionGridResolution)
+{
+	InCollisionGridResolution = FMath::Max(InCollisionGridResolution, 1);
+	if (CollisionGridResolution == InCollisionGridResolution)
+	{
+		return;
+	}
+
+	CollisionGridResolution = InCollisionGridResolution;
+	OnCollisionOptionsChanged();
+}
+
+void ACEClonerActor::SetCollisionGridSize(const FVector& InCollisionGridSize)
+{
+	const FVector NewCollisionGridSize = InCollisionGridSize.ComponentMax(FVector::ZeroVector);
+	if (CollisionGridSize.Equals(NewCollisionGridSize))
+	{
+		return;
+	}
+
+	CollisionGridSize = NewCollisionGridSize;
+	OnCollisionOptionsChanged();
+}
+
+void ACEClonerActor::SetCollisionRadiusMode(ECEClonerCollisionRadiusMode InMode)
+{
+	if (CollisionRadiusMode == InMode)
+	{
+		return;
+	}
+
+	CollisionRadiusMode = InMode;
+	OnCollisionOptionsChanged();
+}
+
+void ACEClonerActor::SetMassMin(float InMassMin)
+{
+	InMassMin = FMath::Max(InMassMin, 1);
+	if (FMath::IsNearlyEqual(MassMin, InMassMin))
+	{
+		return;
+	}
+
+	MassMin = InMassMin;
+	OnCollisionOptionsChanged();
+}
+
+void ACEClonerActor::SetMassMax(float InMassMax)
+{
+	InMassMax = FMath::Max(InMassMax, 1);
+	if (FMath::IsNearlyEqual(MassMax, InMassMax))
+	{
+		return;
+	}
+
+	MassMax = InMassMax;
+	OnCollisionOptionsChanged();
+}
+
 void ACEClonerActor::SetDeltaStepEnabled(bool bInEnabled)
 {
 	if (bDeltaStepEnabled == bInEnabled)
@@ -990,10 +1108,79 @@ void ACEClonerActor::OnLifetimeOptionsChanged()
 	RequestClonerUpdate();
 }
 
+void ACEClonerActor::OnCollisionOptionsChanged()
+{
+	if (!ClonerComponent || !bEnabled)
+	{
+		return;
+	}
+
+	MassMin = FMath::Clamp(MassMin, 1, MassMax);
+	MassMax = FMath::Max(MassMax, MassMin);
+
+	ClonerComponent->SetBoolParameter(TEXT("SurfaceCollisionEnabled"), bSurfaceCollisionEnabled);
+	ClonerComponent->SetIntParameter(TEXT("CollisionIterations"), bParticleCollisionEnabled ? CollisionIterations : 0);
+	ClonerComponent->SetBoolParameter(TEXT("CollisionVelocityEnabled"), bParticleCollisionEnabled ? bCollisionVelocityEnabled : false);
+	ClonerComponent->SetIntParameter(TEXT("CollisionGridResolution"), CollisionGridResolution);
+	ClonerComponent->SetVectorParameter(TEXT("CollisionGridSize"), CollisionGridSize);
+	ClonerComponent->SetFloatParameter(TEXT("MassMin"), MassMin);
+	ClonerComponent->SetFloatParameter(TEXT("MassMax"), MassMax);
+
+	// Adjust size based on attached mesh count
+	CollisionRadii.SetNum(GetMeshCount());
+
+	if (const UCEClonerLayoutBase* LayoutSystem = GetActiveLayout())
+	{
+		if (CollisionRadiusMode != ECEClonerCollisionRadiusMode::Manual)
+		{
+			UNiagaraMeshRendererProperties* MeshRenderer = LayoutSystem->GetMeshRenderer();
+
+			for (int32 Idx = 0; Idx < CollisionRadii.Num(); Idx++)
+			{
+				const FNiagaraMeshRendererMeshProperties& MeshProperties = MeshRenderer->Meshes[Idx];
+				FBoxSphereBounds MeshBounds(ForceInitToZero);
+				const FTransform BoundTransform(MeshProperties.Rotation, MeshProperties.PivotOffset, MeshProperties.Scale);
+
+				if (MeshProperties.Mesh)
+				{
+					MeshBounds = MeshProperties.Mesh->GetBounds().TransformBy(BoundTransform);
+				}
+
+				switch (CollisionRadiusMode)
+				{
+					case ECEClonerCollisionRadiusMode::MinExtent:
+						CollisionRadii[Idx] = MeshBounds.BoxExtent.GetMin();
+					break;
+					case ECEClonerCollisionRadiusMode::MaxExtent:
+						CollisionRadii[Idx] = MeshBounds.BoxExtent.GetMax();
+					break;
+					default:
+					case ECEClonerCollisionRadiusMode::ExtentLength:
+						CollisionRadii[Idx] = MeshBounds.SphereRadius;
+					break;
+				}
+			}
+		}
+
+		const FNiagaraUserRedirectionParameterStore& ExposedParameters = LayoutSystem->GetSystem()->GetExposedParameters();
+
+		static const FNiagaraVariable CollisionRadiiVar(FNiagaraTypeDefinition(UNiagaraDataInterfaceArrayFloat::StaticClass()), TEXT("CollisionRadii"));
+
+		if (UNiagaraDataInterfaceArrayFloat* CollisionRadiiDI = Cast<UNiagaraDataInterfaceArrayFloat>(ExposedParameters.GetDataInterface(CollisionRadiiVar)))
+		{
+			CollisionRadiiDI->GetArrayReference() = CollisionRadii;
+		}
+	}
+
+	RequestClonerUpdate();
+}
+
 void ACEClonerActor::OnClonerMeshUpdated(UCEClonerComponent* InClonerComponent)
 {
 	if (InClonerComponent == ClonerComponent)
 	{
+		OnCollisionOptionsChanged();
+
 		RequestClonerUpdate();
 	}
 }
diff --git a/Engine/Plugins/Experimental/ClonerEffector/Source/ClonerEffector/Private/Cloner/CEClonerComponent.cpp b/Engine/Plugins/Experimental/ClonerEffector/Source/ClonerEffector/Private/Cloner/CEClonerComponent.cpp
--- a/Engine/Plugins/Experimental/ClonerEffector/Source/ClonerEffector/Private/Cloner/CEClonerComponent.cpp
+++ b/Engine/Plugins/Experimental/ClonerEffector/Source/ClonerEffector/Private/Cloner/CEClonerComponent.cpp
@@ -1245,10 +1245,6 @@ void UCEClonerComponent::UpdateRootActorBakedStaticMesh(AActor* InRootActor)
 		EGeometryScriptOutcomePins OutResult;
 		UGeometryScriptLibrary_StaticMeshFunctions::CopyMeshToStaticMesh(MergedAttachmentMesh, BakedAttachmentMesh, AssetOptions, TargetLOD, OutResult);
 
-		// Fix camera culling and avoid killing particles when bounds partially in view
-		const float NewRadiusBounds = BakedAttachmentMesh->GetBounds().SphereRadius * 2;
-		BakedAttachmentMesh->SetExtendedBounds(FBox(-FVector(NewRadiusBounds), FVector(NewRadiusBounds)));
-
 #if WITH_EDITORONLY_DATA
 		// Compute normals and tangents
 		FMeshDescription* MeshDescription = BakedAttachmentMesh->GetMeshDescription(0);
@@ -1494,6 +1490,9 @@ void UCEClonerComponent::UpdateClonerMeshes()
 
 #if WITH_EDITORONLY_DATA
 	MeshRenderer->OnMeshChanged();
+
+	// Used by other data interfaces to update their cached data
+	MeshRenderer->OnChanged().Broadcast();
 #endif
 
 	OnClonerMeshUpdated.Broadcast(this);
diff --git a/Engine/Plugins/Experimental/ClonerEffector/Source/ClonerEffector/Private/Effector/CEEffectorActor.cpp b/Engine/Plugins/Experimental/ClonerEffector/Source/ClonerEffector/Private/Effector/CEEffectorActor.cpp
--- a/Engine/Plugins/Experimental/ClonerEffector/Source/ClonerEffector/Private/Effector/CEEffectorActor.cpp
+++ b/Engine/Plugins/Experimental/ClonerEffector/Source/ClonerEffector/Private/Effector/CEEffectorActor.cpp
@@ -120,6 +120,8 @@ TCEPropertyChangeDispatcher<ACEEffectorActor> ACEEffectorActor::PropertyChangeDi
 	{ GET_MEMBER_NAME_CHECKED(ACEEffectorActor, ScaleStrength), &ACEEffectorActor::OnNoiseFieldOptionsChanged },
 	{ GET_MEMBER_NAME_CHECKED(ACEEffectorActor, Pan), &ACEEffectorActor::OnNoiseFieldOptionsChanged },
 	{ GET_MEMBER_NAME_CHECKED(ACEEffectorActor, Frequency), &ACEEffectorActor::OnNoiseFieldOptionsChanged },
+	{ GET_MEMBER_NAME_CHECKED(ACEEffectorActor, PushStrength), &ACEEffectorActor::OnPushOptionsChanged },
+	{ GET_MEMBER_NAME_CHECKED(ACEEffectorActor, PushDirection), &ACEEffectorActor::OnPushOptionsChanged },
 	/** Force */
 	{ GET_MEMBER_NAME_CHECKED(ACEEffectorActor, bOrientationForceEnabled), &ACEEffectorActor::OnForceEnabledChanged },
 	{ GET_MEMBER_NAME_CHECKED(ACEEffectorActor, OrientationForceRate), &ACEEffectorActor::OnForceOptionsChanged },
@@ -806,6 +808,28 @@ void ACEEffectorActor::SetColor(const FLinearColor& InColor)
 	OnColorChanged();
 }
 
+void ACEEffectorActor::SetPushDirection(ECEClonerEffectorPushDirection InDirection)
+{
+	if (PushDirection == InDirection)
+	{
+		return;
+	}
+
+	PushDirection = InDirection;
+	OnPushOptionsChanged();
+}
+
+void ACEEffectorActor::SetPushStrength(const FVector& InStrength)
+{
+	if (PushStrength.Equals(InStrength))
+	{
+		return;
+	}
+
+	PushStrength = InStrength;
+	OnPushOptionsChanged();
+}
+
 void ACEEffectorActor::OnEffectorTransformed(USceneComponent* InUpdatedComponent, EUpdateTransformFlags InUpdateTransformFlags, ETeleportType InTeleport)
 {
 	OnTransformChanged();
@@ -852,6 +876,7 @@ void ACEEffectorActor::OnModeChanged()
 	OnTransformOptionsChanged();
 	OnTargetActorChanged();
 	OnNoiseFieldOptionsChanged();
+	OnPushOptionsChanged();
 }
 
 void ACEEffectorActor::OnTypeChanged()
@@ -1246,6 +1271,19 @@ void ACEEffectorActor::OnNoiseFieldOptionsChanged()
 	ChannelData.Pan = Pan;
 }
 
+void ACEEffectorActor::OnPushOptionsChanged()
+{
+	if (Mode != ECEClonerEffectorMode::Push)
+	{
+		return;
+	}
+
+	ChannelData.LocationDelta = PushStrength;
+	ChannelData.RotationDelta = FQuat::Identity;
+	ChannelData.ScaleDelta = FVector::OneVector;
+	ChannelData.Pan = FVector(0, static_cast<float>(PushDirection), 0);
+}
+
 void ACEEffectorActor::OnTargetActorTransformChanged(USceneComponent*, EUpdateTransformFlags, ETeleportType)
 {
 	OnTargetOptionsChanged();
diff --git a/Engine/Plugins/Experimental/ClonerEffector/Source/ClonerEffector/Public/CEClonerEffectorShared.h b/Engine/Plugins/Experimental/ClonerEffector/Source/ClonerEffector/Public/CEClonerEffectorShared.h
--- a/Engine/Plugins/Experimental/ClonerEffector/Source/ClonerEffector/Public/CEClonerEffectorShared.h
+++ b/Engine/Plugins/Experimental/ClonerEffector/Source/ClonerEffector/Public/CEClonerEffectorShared.h
@@ -154,7 +154,25 @@ enum class ECEClonerEffectorMode : uint8
 	/** Rotates clones towards a target actor */
 	Target,
 	/** Randomly applies curl noise across the field zone */
-	NoiseField
+	NoiseField,
+	/** Pushes clones apart based on a strength */
+	Push
+};
+
+/** Enumerates the effector push direction available */
+UENUM(BlueprintType)
+enum class ECEClonerEffectorPushDirection : uint8
+{
+	/** Push based on the clone forward vector */
+	Forward,
+	/** Push based on the clone right vector */
+	Right,
+	/** Push based on the cloner up vector */
+	Up,
+	/** Push based on the clone position relative to the effector */
+	Effector,
+	/** Push based on a random unit vector based on the cloner seed */
+	Random
 };
 
 UENUM()
@@ -217,6 +235,20 @@ enum class ECEClonerSpawnBehaviorMode : uint8
 	Rate,
 };
 
+/** Enumerates all modes for how clones radius are calculed */
+UENUM()
+enum class ECEClonerCollisionRadiusMode : uint8
+{
+	/** Input collision radius manually */
+	Manual,
+	/** Collision radius will be calculated automatically based on the min extent value, mesh scale included */
+	MinExtent,
+	/** Collision radius will be calculated automatically based on the max extent value, mesh scale included */
+	MaxExtent,
+	/** Collision radius will be calculated automatically based on the extent length, mesh scale included */
+	ExtentLength
+};
+
 USTRUCT()
 struct FCEClonerAttachmentItem
 {
diff --git a/Engine/Plugins/Experimental/ClonerEffector/Source/ClonerEffector/Public/Cloner/CEClonerActor.h b/Engine/Plugins/Experimental/ClonerEffector/Source/ClonerEffector/Public/Cloner/CEClonerActor.h
--- a/Engine/Plugins/Experimental/ClonerEffector/Source/ClonerEffector/Public/Cloner/CEClonerActor.h
+++ b/Engine/Plugins/Experimental/ClonerEffector/Source/ClonerEffector/Public/Cloner/CEClonerActor.h
@@ -127,6 +127,93 @@ public:
 		return bVisualizeEffectors;
 	}
 
+	UFUNCTION(BlueprintCallable, Category="Cloner")
+	CLONEREFFECTOR_API void SetSurfaceCollisionEnabled(bool bInSurfaceCollisionEnabled);
+
+	UFUNCTION(BlueprintPure, Category="Cloner")
+	bool GetSurfaceCollisionEnabled() const
+	{
+		return bSurfaceCollisionEnabled;
+	}
+
+	UFUNCTION(BlueprintCallable, Category="Cloner")
+	CLONEREFFECTOR_API void SetParticleCollisionEnabled(bool bInParticleCollisionEnabled);
+
+	UFUNCTION(BlueprintPure, Category="Cloner")
+	bool GetParticleCollisionEnabled() const
+	{
+		return bParticleCollisionEnabled;
+	}
+
+	UFUNCTION(BlueprintCallable, Category="Cloner")
+	CLONEREFFECTOR_API void SetCollisionVelocityEnabled(bool bInCollisionVelocityEnabled);
+
+	UFUNCTION(BlueprintPure, Category="Cloner")
+	bool GetCollisionVelocityEnabled() const
+	{
+		return bCollisionVelocityEnabled;
+	}
+
+	UFUNCTION(BlueprintCallable, Category="Cloner")
+	CLONEREFFECTOR_API void SetCollisionIterations(int32 InCollisionIterations);
+
+	UFUNCTION(BlueprintPure, Category="Cloner")
+	int32 GetCollisionIterations() const
+	{
+		return CollisionIterations;
+	}
+
+	UFUNCTION(BlueprintCallable, Category="Cloner")
+	CLONEREFFECTOR_API void SetCollisionGridResolution(int32 InCollisionGridResolution);
+
+	UFUNCTION(BlueprintPure, Category="Cloner")
+	int32 GetCollisionGridResolution() const
+	{
+		return CollisionGridResolution;
+	}
+
+	UFUNCTION(BlueprintCallable, Category="Cloner")
+	CLONEREFFECTOR_API void SetCollisionGridSize(const FVector& InCollisionGridSize);
+
+	UFUNCTION(BlueprintPure, Category="Cloner")
+	FVector GetCollisionGridSize() const
+	{
+		return CollisionGridSize;
+	}
+
+	UFUNCTION(BlueprintCallable, Category="Cloner")
+	CLONEREFFECTOR_API void SetCollisionRadiusMode(ECEClonerCollisionRadiusMode InMode);
+
+	UFUNCTION(BlueprintPure, Category="Cloner")
+	ECEClonerCollisionRadiusMode GetCollisionRadiusMode() const
+	{
+		return CollisionRadiusMode;
+	}
+
+	UFUNCTION(BlueprintPure, Category="Cloner")
+	const TArray<float>& GetCollisionRadii() const
+	{
+		return CollisionRadii;
+	}
+
+	UFUNCTION(BlueprintCallable, Category="Cloner")
+	CLONEREFFECTOR_API void SetMassMin(float InMassMin);
+
+	UFUNCTION(BlueprintPure, Category="Cloner")
+	float GetMassMin() const
+	{
+		return MassMin;
+	}
+
+	UFUNCTION(BlueprintCallable, Category="Cloner")
+	CLONEREFFECTOR_API void SetMassMax(float InMassMax);
+
+	UFUNCTION(BlueprintPure, Category="Cloner")
+	float GetMassMax() const
+	{
+		return MassMax;
+	}
+
 	UFUNCTION(BlueprintCallable, Category="Cloner")
 	CLONEREFFECTOR_API void SetDeltaStepEnabled(bool bInEnabled);
 
@@ -489,6 +576,7 @@ protected:
 	void OnSpawnOptionsChanged();
 	void OnLifetimeScaleCurveChanged();
 	void OnLifetimeOptionsChanged();
+	void OnCollisionOptionsChanged();
 
 	/** Update sprite visibility of this cloner */
 	void OnVisualizerSpriteVisibleChanged();
@@ -544,6 +632,49 @@ protected:
 	UPROPERTY(EditAnywhere, BlueprintReadWrite, Setter, Getter, Category="Renderer", meta=(EditCondition="bUseOverrideMaterial", EditConditionHides))
 	TObjectPtr<UMaterialInterface> OverrideMaterial;
 
+	/**
+	 * Allow particles to react to surface by using the distance field,
+	 * ensure the mesh you want particle to collide with is tick enough or increase its distance field resolution scale in static mesh editor 
+	 */
+	UPROPERTY(EditAnywhere, BlueprintReadWrite, Setter="SetSurfaceCollisionEnabled", Getter="GetSurfaceCollisionEnabled", Category="Collisions")
+	bool bSurfaceCollisionEnabled = false;
+
+	/** Allow particles to react to other emitter particles, uses a neighbor grid to detect collision */
+	UPROPERTY(EditAnywhere, BlueprintReadWrite, Setter="SetParticleCollisionEnabled", Getter="GetParticleCollisionEnabled", Category="Collisions")
+	bool bParticleCollisionEnabled = false;
+
+	/** Recalculate accurate velocity after collision is updated */
+	UPROPERTY(EditAnywhere, BlueprintReadWrite, Setter="SetCollisionVelocityEnabled", Getter="GetCollisionVelocityEnabled", Category="Collisions", meta=(EditCondition="bParticleCollisionEnabled", EditConditionHides))
+	bool bCollisionVelocityEnabled = false;
+
+	/** Amount of iterations to improve particle collision results but affects performance */
+	UPROPERTY(EditAnywhere, BlueprintReadWrite, Setter, Getter, Category="Collisions", meta=(ClampMin="1", EditCondition="bParticleCollisionEnabled", EditConditionHides))
+	int32 CollisionIterations = 1;
+
+	/** Resolution of the neighbor grid to detect collision between particles */
+	UPROPERTY(EditAnywhere, BlueprintReadWrite, Setter, Getter, Category="Collisions", meta=(ClampMin="1", EditCondition="bParticleCollisionEnabled", EditConditionHides))
+	int32 CollisionGridResolution = 32;
+
+	/** Size of the neighbor grid to detect collision between particles, particles outside this grid will not have collisions */
+	UPROPERTY(EditAnywhere, BlueprintReadWrite, Setter, Getter, Category="Collisions", meta=(ClampMin="0", AllowPreserveRatio, EditCondition="bParticleCollisionEnabled", EditConditionHides))
+	FVector CollisionGridSize = FVector(5000.f);
+
+	/** Collision radius calculation mode */
+	UPROPERTY(EditAnywhere, BlueprintReadWrite, Setter, Getter, Category="Collisions", meta=(EditCondition="bParticleCollisionEnabled", EditConditionHides))
+	ECEClonerCollisionRadiusMode CollisionRadiusMode = ECEClonerCollisionRadiusMode::ExtentLength;
+
+	/** Radius expected around each particle for collision, order matches attachment index */
+	UPROPERTY(EditAnywhere, BlueprintReadWrite, Getter, Category="Collisions", EditFixedSize, meta=(ClampMin="0", EditFixedOrder, EditCondition="bParticleCollisionEnabled && CollisionRadiusMode == ECEClonerCollisionRadiusMode::Manual", EditConditionHides))
+	TArray<float> CollisionRadii;
+
+	/** Minimum particle mass, used for collisions to push apart */
+	UPROPERTY(EditAnywhere, BlueprintReadWrite, Setter, Getter, Category="Collisions", meta=(ClampMin="1"))
+	float MassMin = 1.f;
+
+	/** Maximum particle mass, used for collisions to push apart */
+	UPROPERTY(EditAnywhere, BlueprintReadWrite, Setter, Getter, Category="Collisions", meta=(ClampMin="1"))
+	float MassMax = 2.f;
+
 	/** How many times do we spawn clones */
 	UPROPERTY(EditAnywhere, BlueprintReadWrite, Setter="SetSpawnLoopMode", Getter="GetSpawnLoopMode", DisplayName="LoopMode", Category="Spawn")
 	ECEClonerSpawnLoopMode SpawnLoopMode = ECEClonerSpawnLoopMode::Once;
diff --git a/Engine/Plugins/Experimental/ClonerEffector/Source/ClonerEffector/Public/Effector/CEEffectorActor.h b/Engine/Plugins/Experimental/ClonerEffector/Source/ClonerEffector/Public/Effector/CEEffectorActor.h
--- a/Engine/Plugins/Experimental/ClonerEffector/Source/ClonerEffector/Public/Effector/CEEffectorActor.h
+++ b/Engine/Plugins/Experimental/ClonerEffector/Source/ClonerEffector/Public/Effector/CEEffectorActor.h
@@ -411,6 +411,24 @@ public:
 		return Color;
 	}
 
+	UFUNCTION(BlueprintCallable, Category="Effector")
+	CLONEREFFECTOR_API void SetPushDirection(ECEClonerEffectorPushDirection InDirection);
+
+	UFUNCTION(BlueprintPure, Category="Effector")
+	ECEClonerEffectorPushDirection GetPushDirection() const
+	{
+		return PushDirection;
+	}
+
+	UFUNCTION(BlueprintCallable, Category="Effector")
+	CLONEREFFECTOR_API void SetPushStrength(const FVector& InStrength);
+
+	UFUNCTION(BlueprintPure, Category="Effector")
+	const FVector& GetPushStrength() const
+	{
+		return PushStrength;
+	}
+
 #if WITH_EDITOR
 	UFUNCTION(BlueprintCallable, Category="Effector")
 	CLONEREFFECTOR_API void SetVisualizerComponentVisible(bool bInVisible);
@@ -519,6 +537,8 @@ protected:
 	/** Called when noise field mode options are changed */
 	void OnNoiseFieldOptionsChanged();
 
+	void OnPushOptionsChanged();
+
 	/** Called when force options are changed */
 	void OnForceOptionsChanged();
 
@@ -555,59 +575,59 @@ protected:
 	FLinearColor Color = FLinearColor::Red;
 
 	/** Type of effector to apply on cloners instances */
-	UPROPERTY(EditAnywhere, BlueprintReadWrite, Setter, Getter, Category="Type")
+	UPROPERTY(EditAnywhere, BlueprintReadWrite, Setter, Getter, Category="Shape")
 	ECEClonerEffectorType Type = ECEClonerEffectorType::Sphere;
 
 	/** Invert the type effect, instead of affecting the inside of a zone, will affect the outside */
-	UPROPERTY(EditAnywhere, BlueprintReadWrite, Setter="SetInvertType", Getter="GetInvertType", Category="Type")
+	UPROPERTY(EditAnywhere, BlueprintReadWrite, Setter="SetInvertType", Getter="GetInvertType", Category="Shape")
 	bool bInvertType = false;
 
 	/** Weight easing function applied to lerp transforms */
-	UPROPERTY(EditAnywhere, BlueprintReadWrite, Setter, Getter, Category="Type", meta=(EditCondition="Type != ECEClonerEffectorType::Unbound", EditConditionHides))
+	UPROPERTY(EditAnywhere, BlueprintReadWrite, Setter, Getter, Category="Shape", meta=(EditCondition="Type != ECEClonerEffectorType::Unbound", EditConditionHides))
 	ECEClonerEasing Easing = ECEClonerEasing::Linear;
 
 	/** Inner radius of sphere, all clones inside will be affected with a maximum weight */
-	UPROPERTY(EditAnywhere, BlueprintReadWrite, Setter, Getter, Category="Type", meta=(ClampMin="0", EditCondition="Type == ECEClonerEffectorType::Sphere", EditConditionHides))
+	UPROPERTY(EditAnywhere, BlueprintReadWrite, Setter, Getter, Category="Shape", meta=(ClampMin="0", EditCondition="Type == ECEClonerEffectorType::Sphere", EditConditionHides))
 	float InnerRadius = 50.f;
 
 	/** Outer radius of sphere, all clones outside will not be affected */
-	UPROPERTY(EditAnywhere, BlueprintReadWrite, Setter, Getter, Category="Type", meta=(ClampMin="0", EditCondition="Type == ECEClonerEffectorType::Sphere", EditConditionHides))
+	UPROPERTY(EditAnywhere, BlueprintReadWrite, Setter, Getter, Category="Shape", meta=(ClampMin="0", EditCondition="Type == ECEClonerEffectorType::Sphere", EditConditionHides))
 	float OuterRadius = 200.f;
 
 	/** Inner extent of box, all clones inside will be affected with a maximum weight */
-	UPROPERTY(EditAnywhere, BlueprintReadWrite, Setter, Getter, Category="Type", meta=(ClampMin="0", AllowPreserveRatio, EditCondition="Type == ECEClonerEffectorType::Box", EditConditionHides))
+	UPROPERTY(EditAnywhere, BlueprintReadWrite, Setter, Getter, Category="Shape", meta=(ClampMin="0", AllowPreserveRatio, EditCondition="Type == ECEClonerEffectorType::Box", EditConditionHides))
 	FVector InnerExtent = FVector(50.f);
 
 	/** Outer extent of box, all clones outside will not be affected */
-	UPROPERTY(EditAnywhere, BlueprintReadWrite, Setter, Getter, Category="Type", meta=(ClampMin="0", AllowPreserveRatio, EditCondition="Type == ECEClonerEffectorType::Box", EditConditionHides))
+	UPROPERTY(EditAnywhere, BlueprintReadWrite, Setter, Getter, Category="Shape", meta=(ClampMin="0", AllowPreserveRatio, EditCondition="Type == ECEClonerEffectorType::Box", EditConditionHides))
 	FVector OuterExtent = FVector(200.f);
 
 	/** Plane spacing, everything inside this zone will be affected */
-	UPROPERTY(EditAnywhere, BlueprintReadWrite, Setter, Getter, Category="Type", meta=(ClampMin="0", EditCondition="Type == ECEClonerEffectorType::Plane", EditConditionHides))
+	UPROPERTY(EditAnywhere, BlueprintReadWrite, Setter, Getter, Category="Shape", meta=(ClampMin="0", EditCondition="Type == ECEClonerEffectorType::Plane", EditConditionHides))
 	float PlaneSpacing = 200.f;
 
 	/** Radial angle in degree, everything within the angle will be affected */
-	UPROPERTY(EditAnywhere, BlueprintReadWrite, Setter, Getter, Category="Type", meta=(ClampMin="0", ClampMax="360", EditCondition="Type == ECEClonerEffectorType::Radial", EditConditionHides))
+	UPROPERTY(EditAnywhere, BlueprintReadWrite, Setter, Getter, Category="Shape", meta=(ClampMin="0", ClampMax="360", EditCondition="Type == ECEClonerEffectorType::Radial", EditConditionHides))
 	float RadialAngle = 180.f;
 
 	/** Minimum radius for the radial effect to be applied on clones, below clones will not be affected */
-	UPROPERTY(EditAnywhere, BlueprintReadWrite, Setter, Getter, Category="Type", meta=(ClampMin="0", EditCondition="Type == ECEClonerEffectorType::Radial", EditConditionHides))
+	UPROPERTY(EditAnywhere, BlueprintReadWrite, Setter, Getter, Category="Shape", meta=(ClampMin="0", EditCondition="Type == ECEClonerEffectorType::Radial", EditConditionHides))
 	float RadialMinRadius = 0.f;
 
 	/** Maximum radius for the radial effect to be applied on clones, above clones will not be affected */
-	UPROPERTY(EditAnywhere, BlueprintReadWrite, Setter, Getter, Category="Type", meta=(ClampMin="0", EditCondition="Type == ECEClonerEffectorType::Radial", EditConditionHides))
+	UPROPERTY(EditAnywhere, BlueprintReadWrite, Setter, Getter, Category="Shape", meta=(ClampMin="0", EditCondition="Type == ECEClonerEffectorType::Radial", EditConditionHides))
 	float RadialMaxRadius = 1000.f;
 
 	/** Main torus radius from center to the edge where inner and outer tube will be revolved */
-	UPROPERTY(EditAnywhere, BlueprintReadWrite, Setter, Getter, Category="Type", meta=(ClampMin="0", EditCondition="Type == ECEClonerEffectorType::Torus", EditConditionHides))
+	UPROPERTY(EditAnywhere, BlueprintReadWrite, Setter, Getter, Category="Shape", meta=(ClampMin="0", EditCondition="Type == ECEClonerEffectorType::Torus", EditConditionHides))
 	float TorusRadius = 250.f;
 
 	/** Minimum revolved radius for the torus effect, clones contained inside will be affected with a maximum weight */
-	UPROPERTY(EditAnywhere, BlueprintReadWrite, Setter, Getter, Category="Type", meta=(ClampMin="0", EditCondition="Type == ECEClonerEffectorType::Torus", EditConditionHides))
+	UPROPERTY(EditAnywhere, BlueprintReadWrite, Setter, Getter, Category="Shape", meta=(ClampMin="0", EditCondition="Type == ECEClonerEffectorType::Torus", EditConditionHides))
 	float TorusInnerRadius = 50.f;
 
 	/** Maximum revolved radius for the torus effect, clones outside of it will not be affected */
-	UPROPERTY(EditAnywhere, BlueprintReadWrite, Setter, Getter, Category="Type", meta=(ClampMin="0", EditCondition="Type == ECEClonerEffectorType::Torus", EditConditionHides))
+	UPROPERTY(EditAnywhere, BlueprintReadWrite, Setter, Getter, Category="Shape", meta=(ClampMin="0", EditCondition="Type == ECEClonerEffectorType::Torus", EditConditionHides))
 	float TorusOuterRadius = 200.f;
 
 	/** Mode of effector for each clones instances */
@@ -653,6 +673,14 @@ protected:
 	UPROPERTY(EditAnywhere, BlueprintReadWrite, Setter, Getter, Category="Mode", meta=(ClampMin="0", EditCondition="Mode == ECEClonerEffectorMode::NoiseField", EditConditionHides))
 	float Frequency = 0.5f;
 
+	/** Strength of the push effect */
+	UPROPERTY(EditAnywhere, BlueprintReadWrite, Setter, Getter, Category="Mode", meta=(AllowPreserveRatio, EditCondition="Mode == ECEClonerEffectorMode::Push", EditConditionHides))
+	FVector PushStrength = FVector(100.f);
+
+	/** Relative direction computed for the push effect on each clone */
+	UPROPERTY(EditAnywhere, BlueprintReadWrite, Setter, Getter, Category="Mode", meta=(EditCondition="Mode == ECEClonerEffectorMode::Push", EditConditionHides))
+	ECEClonerEffectorPushDirection PushDirection = ECEClonerEffectorPushDirection::Forward;
+
 	/** Enable orientation force to allow each clone instance to rotate around its pivot */
 	UPROPERTY(EditAnywhere, BlueprintReadWrite, Setter="SetOrientationForceEnabled", Getter="GetOrientationForceEnabled", Category="Force")
 	bool bOrientationForceEnabled = false;
diff --git a/Engine/Plugins/Experimental/ClonerEffector/Source/ClonerEffectorEditor/Private/Cloner/CEEditorClonerDetailCustomization.cpp b/Engine/Plugins/Experimental/ClonerEffector/Source/ClonerEffectorEditor/Private/Cloner/CEEditorClonerDetailCustomization.cpp
--- a/Engine/Plugins/Experimental/ClonerEffector/Source/ClonerEffectorEditor/Private/Cloner/CEEditorClonerDetailCustomization.cpp
+++ b/Engine/Plugins/Experimental/ClonerEffector/Source/ClonerEffectorEditor/Private/Cloner/CEEditorClonerDetailCustomization.cpp
@@ -23,6 +23,7 @@ void FCEEditorClonerDetailCustomization::CustomizeDetails(IDetailLayoutBuilder&
 		TEXT("Progress"),
 		TEXT("Step"),
 		TEXT("Range"),
+		TEXT("Collisions")
 		TEXT("Spawn"),
 		TEXT("Lifetime"),
 	};
@@ -195,6 +196,9 @@ void FCEEditorClonerDetailCustomization::RegisterCustomSections() const
 
 	const TSharedRef<FPropertySection> LifetimeSection = PropertyModule.FindOrCreateSection(ClassName, "Lifetime", LOCTEXT("Cloner.Lifetime", "Lifetime"));
 	LifetimeSection->AddCategory("Lifetime");
+
+	const TSharedRef<FPropertySection> CollisionSection = PropertyModule.FindOrCreateSection(ClassName, "Collisions", LOCTEXT("Cloner.Collisions", "Collisions"));
+	CollisionSection->AddCategory("Collisions");
 }
 
 #undef LOCTEXT_NAMESPACE
diff --git a/Engine/Plugins/Experimental/ClonerEffector/Source/ClonerEffectorEditor/Private/Effector/CEEditorEffectorDetailCustomization.cpp b/Engine/Plugins/Experimental/ClonerEffector/Source/ClonerEffectorEditor/Private/Effector/CEEditorEffectorDetailCustomization.cpp
--- a/Engine/Plugins/Experimental/ClonerEffector/Source/ClonerEffectorEditor/Private/Effector/CEEditorEffectorDetailCustomization.cpp
+++ b/Engine/Plugins/Experimental/ClonerEffector/Source/ClonerEffectorEditor/Private/Effector/CEEditorEffectorDetailCustomization.cpp
@@ -79,8 +79,8 @@ void FCEEditorEffectorDetailCustomization::RegisterCustomSections() const
 	const TSharedRef<FPropertySection> EffectorSection = PropertyModule.FindOrCreateSection(ClassName, "General", LOCTEXT("Effector.General", "General"));
 	EffectorSection->AddCategory(TEXT("Effector"));
 
-	const TSharedRef<FPropertySection> TypeSection = PropertyModule.FindOrCreateSection(ClassName, "Type", LOCTEXT("Effector.Type", "Type"));
-	TypeSection->AddCategory(TEXT("Type"));
+	const TSharedRef<FPropertySection> ShapeSection = PropertyModule.FindOrCreateSection(ClassName, "Shape", LOCTEXT("Effector.Shape", "Shape"));
+	ShapeSection->AddCategory(TEXT("Shape"));
 
 	const TSharedRef<FPropertySection> ModeSection = PropertyModule.FindOrCreateSection(ClassName, "Mode", LOCTEXT("Effector.Mode", "Mode"));
 	ModeSection->AddCategory(TEXT("Mode"));
