MotionDesign : ClonerEffector

diff --git a/Engine/Plugins/Experimental/ClonerEffector/Source/ClonerEffector/Private/Cloner/CEClonerActor.cpp b/Engine/Plugins/Experimental/ClonerEffector/Source/ClonerEffector/Private/Cloner/CEClonerActor.cpp
--- a/Engine/Plugins/Experimental/ClonerEffector/Source/ClonerEffector/Private/Cloner/CEClonerActor.cpp
+++ b/Engine/Plugins/Experimental/ClonerEffector/Source/ClonerEffector/Private/Cloner/CEClonerActor.cpp
@@ -3,18 +3,26 @@
 #include "Cloner/CEClonerActor.h"
 
 #include "Cloner/CEClonerComponent.h"
+#include "Cloner/Extensions/CEClonerCollisionExtension.h"
+#include "Cloner/Extensions/CEClonerEffectorExtension.h"
+#include "Cloner/Extensions/CEClonerEmitterSpawnExtension.h"
+#include "Cloner/Extensions/CEClonerLifetimeExtension.h"
+#include "Cloner/Extensions/CEClonerMeshRendererExtension.h"
+#include "Cloner/Extensions/CEClonerProgressExtension.h"
+#include "Cloner/Extensions/CEClonerRangeExtension.h"
+#include "Cloner/Extensions/CEClonerStepExtension.h"
+#include "Cloner/Layouts/CEClonerCircleLayout.h"
+#include "Cloner/Layouts/CEClonerCylinderLayout.h"
+#include "Cloner/Layouts/CEClonerGridLayout.h"
+#include "Cloner/Layouts/CEClonerHoneycombLayout.h"
 #include "Cloner/Layouts/CEClonerLayoutBase.h"
+#include "Cloner/Layouts/CEClonerLineLayout.h"
+#include "Cloner/Layouts/CEClonerMeshLayout.h"
+#include "Cloner/Layouts/CEClonerSphereRandomLayout.h"
+#include "Cloner/Layouts/CEClonerSphereUniformLayout.h"
+#include "Cloner/Layouts/CEClonerSplineLayout.h"
 #include "Effector/CEEffectorActor.h"
-#include "Engine/StaticMeshActor.h"
 #include "Materials/MaterialInterface.h"
-#include "NiagaraDataInterfaceArrayFloat.h"
-#include "NiagaraDataInterfaceCurve.h"
-#include "NiagaraDataInterfaceSkeletalMesh.h"
-#include "NiagaraMeshRendererProperties.h"
-#include "NiagaraSystem.h"
-#include "Settings/CEClonerEffectorSettings.h"
-#include "Subsystems/CEClonerSubsystem.h"
-#include "Subsystems/CEEffectorSubsystem.h"
 #include "UObject/ConstructorHelpers.h"
 
 #if WITH_EDITOR
@@ -25,1480 +33,329 @@
 
 DEFINE_LOG_CATEGORY_STATIC(LogCEClonerActor, Log, All);
 
-// Sets default values
-ACEClonerActor::ACEClonerActor()
-{
-	SetCanBeDamaged(false);
-	PrimaryActorTick.bCanEverTick          = true;
-	PrimaryActorTick.bStartWithTickEnabled = true;
-	PrimaryActorTick.bTickEvenWhenPaused   = true;
-	PrimaryActorTick.bHighPriority         = true;
-
-	ClonerComponent = CreateDefaultSubobject<UCEClonerComponent>(TEXT("AvaClonerComponent"));
-	SetRootComponent(ClonerComponent);
-
-	// Default Scale Curve
-	LifetimeScaleCurve.AddKey(0, 1.f);
-	LifetimeScaleCurve.AddKey(1, 0.f);
-
-	// Default override material
-	static ConstructorHelpers::FObjectFinder<UMaterialInterface> DefaultMaterialFinder(UCEClonerEffectorSettings::DefaultMaterialPath);
-	OverrideMaterial = DefaultMaterialFinder.Object;
-
-	if (!IsTemplate())
-	{
-#if WITH_EDITOR
-		if (GEditor)
-		{
-			GEditor->GetSelectedActors()->SelectionChangedEvent.AddUObject(this, &ACEClonerActor::OnEditorSelectionChanged);
-		}
-#endif
-
-		ClonerComponent->OnClonerMeshUpdated.AddUObject(this, &ACEClonerActor::OnClonerMeshUpdated);
-
-		UCEEffectorSubsystem::OnEffectorIdentifierChangedDelegate.AddUObject(this, &ACEClonerActor::OnEffectorIdentifierChanged);
-
-		const TArray<FString> LayoutNames = GetClonerLayoutNames();
-
-		// Apply default layout
-		LayoutName = !LayoutNames.IsEmpty() ? FName(LayoutNames[0]) : NAME_None;
-
-#if WITH_EDITOR
-		if (const UCEClonerSubsystem* ClonerSubsystem = UCEClonerSubsystem::Get())
-		{
-			UCEClonerSubsystem::OnCVarChanged().AddUObject(this, &ACEClonerActor::OnCVarChanged);
-			bReduceMotionGhosting = ClonerSubsystem->IsNoFlickerEnabled();
-		}
-#endif
-	}
-}
-
-#if WITH_EDITOR
-FString ACEClonerActor::GetDefaultActorLabel() const
-{
-	return DefaultLabel;
-}
-
-void ACEClonerActor::PostEditUndo()
-{
-	Super::PostEditUndo();
-
-	ForceUpdateCloner();
-}
-
-const TCEPropertyChangeDispatcher<ACEClonerActor> ACEClonerActor::PropertyChangeDispatcher =
-{
-	{ GET_MEMBER_NAME_CHECKED(ACEClonerActor, bEnabled), &ACEClonerActor::OnEnabledChanged },
-	{ GET_MEMBER_NAME_CHECKED(ACEClonerActor, Seed), &ACEClonerActor::OnSeedChanged },
-	{ GET_MEMBER_NAME_CHECKED(ACEClonerActor, Color), &ACEClonerActor::OnColorChanged },
-	{ GET_MEMBER_NAME_CHECKED(ACEClonerActor, bVisualizeEffectors), &ACEClonerActor::OnOverrideMaterialChanged },
-	{ GET_MEMBER_NAME_CHECKED(ACEClonerActor, EffectorsWeak), &ACEClonerActor::OnEffectorsChanged },
-	/** Collision */
-	{ GET_MEMBER_NAME_CHECKED(ACEClonerActor, bSurfaceCollisionEnabled), &ACEClonerActor::OnCollisionOptionsChanged },
-	{ GET_MEMBER_NAME_CHECKED(ACEClonerActor, bParticleCollisionEnabled), &ACEClonerActor::OnCollisionOptionsChanged },
-	{ GET_MEMBER_NAME_CHECKED(ACEClonerActor, bCollisionVelocityEnabled), &ACEClonerActor::OnCollisionOptionsChanged },
-	{ GET_MEMBER_NAME_CHECKED(ACEClonerActor, CollisionRadiusMode), &ACEClonerActor::OnCollisionOptionsChanged },
-	{ GET_MEMBER_NAME_CHECKED(ACEClonerActor, CollisionRadii), &ACEClonerActor::OnCollisionOptionsChanged },
-	{ GET_MEMBER_NAME_CHECKED(ACEClonerActor, CollisionIterations), &ACEClonerActor::OnCollisionOptionsChanged },
-	{ GET_MEMBER_NAME_CHECKED(ACEClonerActor, CollisionGridResolution), &ACEClonerActor::OnCollisionOptionsChanged },
-	{ GET_MEMBER_NAME_CHECKED(ACEClonerActor, CollisionGridSize), &ACEClonerActor::OnCollisionOptionsChanged },
-	{ GET_MEMBER_NAME_CHECKED(ACEClonerActor, MassMin), &ACEClonerActor::OnCollisionOptionsChanged },
-	{ GET_MEMBER_NAME_CHECKED(ACEClonerActor, MassMax), &ACEClonerActor::OnCollisionOptionsChanged },
-	/** Layout */
-	{ GET_MEMBER_NAME_CHECKED(ACEClonerActor, LayoutName), &ACEClonerActor::OnLayoutNameChanged },
-	/** Advanced */
-	{ GET_MEMBER_NAME_CHECKED(ACEClonerActor, bVisualizerSpriteVisible), &ACEClonerActor::OnVisualizerSpriteVisibleChanged },
-	{ GET_MEMBER_NAME_CHECKED(ACEClonerActor, bReduceMotionGhosting), &ACEClonerActor::OnReduceMotionGhostingChanged },
-	/** Spawn */
-	{ GET_MEMBER_NAME_CHECKED(ACEClonerActor, SpawnLoopMode), &ACEClonerActor::OnSpawnOptionsChanged },
-	{ GET_MEMBER_NAME_CHECKED(ACEClonerActor, SpawnLoopInterval), &ACEClonerActor::OnSpawnOptionsChanged },
-	{ GET_MEMBER_NAME_CHECKED(ACEClonerActor, SpawnLoopIterations), &ACEClonerActor::OnSpawnOptionsChanged },
-	{ GET_MEMBER_NAME_CHECKED(ACEClonerActor, SpawnBehaviorMode), &ACEClonerActor::OnSpawnOptionsChanged },
-	{ GET_MEMBER_NAME_CHECKED(ACEClonerActor, SpawnRate), &ACEClonerActor::OnSpawnOptionsChanged },
-	/** Lifetime */
-	{ GET_MEMBER_NAME_CHECKED(ACEClonerActor, bLifetimeEnabled), &ACEClonerActor::OnLifetimeOptionsChanged },
-	{ GET_MEMBER_NAME_CHECKED(ACEClonerActor, LifetimeMin), &ACEClonerActor::OnLifetimeOptionsChanged },
-	{ GET_MEMBER_NAME_CHECKED(ACEClonerActor, LifetimeMax), &ACEClonerActor::OnLifetimeOptionsChanged },
-	{ GET_MEMBER_NAME_CHECKED(ACEClonerActor, bLifetimeScaleEnabled), &ACEClonerActor::OnLifetimeOptionsChanged },
-	/** Renderer */
-	{ GET_MEMBER_NAME_CHECKED(ACEClonerActor, MeshRenderMode), &ACEClonerActor::OnMeshRenderModeChanged },
-	{ GET_MEMBER_NAME_CHECKED(ACEClonerActor, MeshFacingMode), &ACEClonerActor::OnMeshRendererOptionsChanged },
-	{ GET_MEMBER_NAME_CHECKED(ACEClonerActor, bMeshCastShadows), &ACEClonerActor::OnMeshRendererOptionsChanged },
-	{ GET_MEMBER_NAME_CHECKED(ACEClonerActor, DefaultMeshes), &ACEClonerActor::OnDefaultMeshesChanged },
-	{ GET_MEMBER_NAME_CHECKED(ACEClonerActor, bUseOverrideMaterial), &ACEClonerActor::OnOverrideMaterialChanged },
-	{ GET_MEMBER_NAME_CHECKED(ACEClonerActor, OverrideMaterial), &ACEClonerActor::OnOverrideMaterialChanged },
-	/** Progress */
-	{ GET_MEMBER_NAME_CHECKED(ACEClonerActor, bInvertProgress), &ACEClonerActor::OnProgressChanged },
-	{ GET_MEMBER_NAME_CHECKED(ACEClonerActor, Progress), &ACEClonerActor::OnProgressChanged },
-	/** Step */
-	{ GET_MEMBER_NAME_CHECKED(ACEClonerActor, bDeltaStepEnabled), &ACEClonerActor::OnDeltaStepChanged },
-	{ GET_MEMBER_NAME_CHECKED(ACEClonerActor, DeltaStepRotation), &ACEClonerActor::OnDeltaStepChanged },
-	{ GET_MEMBER_NAME_CHECKED(ACEClonerActor, DeltaStepScale), &ACEClonerActor::OnDeltaStepChanged },
-	/** Range */
-	{ GET_MEMBER_NAME_CHECKED(ACEClonerActor, bRangeEnabled), &ACEClonerActor::OnRangeOptionsChanged },
-	{ GET_MEMBER_NAME_CHECKED(ACEClonerActor, RangeOffsetMin), &ACEClonerActor::OnRangeOptionsChanged },
-	{ GET_MEMBER_NAME_CHECKED(ACEClonerActor, RangeOffsetMax), &ACEClonerActor::OnRangeOptionsChanged },
-	{ GET_MEMBER_NAME_CHECKED(ACEClonerActor, RangeRotationMin), &ACEClonerActor::OnRangeOptionsChanged },
-	{ GET_MEMBER_NAME_CHECKED(ACEClonerActor, RangeRotationMax), &ACEClonerActor::OnRangeOptionsChanged },
-	{ GET_MEMBER_NAME_CHECKED(ACEClonerActor, RangeScaleMin), &ACEClonerActor::OnRangeOptionsChanged },
-	{ GET_MEMBER_NAME_CHECKED(ACEClonerActor, RangeScaleMax), &ACEClonerActor::OnRangeOptionsChanged },
-	{ GET_MEMBER_NAME_CHECKED(ACEClonerActor, bRangeScaleUniform), &ACEClonerActor::OnRangeOptionsChanged },
-	{ GET_MEMBER_NAME_CHECKED(ACEClonerActor, RangeScaleUniformMin), &ACEClonerActor::OnRangeOptionsChanged },
-	{ GET_MEMBER_NAME_CHECKED(ACEClonerActor, RangeScaleUniformMax), &ACEClonerActor::OnRangeOptionsChanged },
-};
-
-void ACEClonerActor::PostEditChangeProperty(FPropertyChangedEvent& InPropertyChangedEvent)
-{
-	Super::PostEditChangeProperty(InPropertyChangedEvent);
-
-	PropertyChangeDispatcher.OnPropertyChanged(this, InPropertyChangedEvent);
-}
-#endif
-
-void ACEClonerActor::PostInitProperties()
-{
-	Super::PostInitProperties();
-
-	if (HasAnyFlags(RF_ClassDefaultObject))
-	{
-		// Register new type def for niagara
-
-		constexpr ENiagaraTypeRegistryFlags MeshFlags =
-			ENiagaraTypeRegistryFlags::AllowAnyVariable |
-			ENiagaraTypeRegistryFlags::AllowParameter;
-
-		FNiagaraTypeRegistry::Register(FNiagaraTypeDefinition(StaticEnum<ECEClonerMeshRenderMode>()), MeshFlags);
-		FNiagaraTypeRegistry::Register(FNiagaraTypeDefinition(StaticEnum<ECEClonerGridConstraint>()), MeshFlags);
-		FNiagaraTypeRegistry::Register(FNiagaraTypeDefinition(StaticEnum<ECEClonerPlane>()), MeshFlags);
-		FNiagaraTypeRegistry::Register(FNiagaraTypeDefinition(StaticEnum<ECEClonerAxis>()), MeshFlags);
-		FNiagaraTypeRegistry::Register(FNiagaraTypeDefinition(StaticEnum<ECEClonerEasing>()), MeshFlags);
-		FNiagaraTypeRegistry::Register(FNiagaraTypeDefinition(StaticEnum<ECEClonerMeshAsset>()), MeshFlags);
-		FNiagaraTypeRegistry::Register(FNiagaraTypeDefinition(StaticEnum<ECEClonerMeshSampleData>()), MeshFlags);
-		FNiagaraTypeRegistry::Register(FNiagaraTypeDefinition(StaticEnum<ECEClonerEffectorType>()), MeshFlags);
-		FNiagaraTypeRegistry::Register(FNiagaraTypeDefinition(StaticEnum<ECEClonerTextureSampleChannel>()), MeshFlags);
-		FNiagaraTypeRegistry::Register(FNiagaraTypeDefinition(StaticEnum<ECEClonerCompareMode>()), MeshFlags);
-		FNiagaraTypeRegistry::Register(FNiagaraTypeDefinition(StaticEnum<ECEClonerEffectorMode>()), MeshFlags);
-		FNiagaraTypeRegistry::Register(FNiagaraTypeDefinition(StaticEnum<ECEClonerSpawnLoopMode>()), MeshFlags);
-		FNiagaraTypeRegistry::Register(FNiagaraTypeDefinition(StaticEnum<ECEClonerSpawnBehaviorMode>()), MeshFlags);
-		FNiagaraTypeRegistry::Register(FNiagaraTypeDefinition(StaticEnum<ECEClonerEffectorPushDirection>()), MeshFlags);
-	}
-}
-
-void ACEClonerActor::PostLoad()
-{
-	Super::PostLoad();
-
-	for (TMap<FName, TObjectPtr<UCEClonerLayoutBase>>::TIterator It(LayoutInstances); It; ++It)
-	{
-		if (!It->Value)
-		{
-			It.RemoveCurrent();
-			continue;
-		}
-
-		// Update layout name if changed
-		const FName NewLayoutName = It->Value->GetLayoutName();
-
-		if (It->Key != NewLayoutName)
-		{
-			LayoutInstances.Add(NewLayoutName, It->Value);
-			It.RemoveCurrent();
-
-			UE_LOG(LogCEClonerActor, Log, TEXT("%s : Cloner layout name changed %s"), *GetActorNameOrLabel(), *NewLayoutName.ToString());
-		}
-	}
-}
-
-void ACEClonerActor::PostActorCreated()
-{
-	Super::PostActorCreated();
-
-#if WITH_EDITOR
-	bSpawnDefaultActorAttached = true;
-#endif
-}
-
-void ACEClonerActor::Tick(float InDeltaTime)
-{
-	Super::Tick(InDeltaTime);
-
-	if (!bClonerInitialized)
-	{
-		InitializeCloner();
-	}
-
-	if (!bEnabled)
-	{
-		return;
-	}
-
-	TreeUpdateDeltaTime += InDeltaTime;
-	if (TreeUpdateDeltaTime >= TreeUpdateInterval)
-	{
-		TreeUpdateDeltaTime -= TreeUpdateInterval;
-
-		if (ClonerComponent)
-		{
-			ClonerComponent->UpdateClonerAttachmentTree();
-			ClonerComponent->UpdateClonerRenderState();
-		}
-	}
-
-	if (bNeedsRefresh)
-	{
-		bNeedsRefresh = false;
-		RequestClonerUpdate(true);
-	}
-}
-
-void ACEClonerActor::UpdateLayoutOptions()
-{
-	OnMeshRenderModeChanged();
-	OnSeedChanged();
-	OnColorChanged();
-
-	// update layouts options
-	if (UCEClonerLayoutBase* CurrentLayout = ClonerComponent->GetClonerActiveLayout())
-	{
-		CurrentLayout->UpdateLayoutParameters();
-	}
-
-	OnMeshRendererOptionsChanged();
-	OnDefaultMeshesChanged();
-	OnOverrideMaterialChanged();
-
-	OnDeltaStepChanged();
-	OnRangeOptionsChanged();
-	OnSpawnOptionsChanged();
-	OnLifetimeOptionsChanged();
-	OnCollisionOptionsChanged();
-}
-
-void ACEClonerActor::SetTreeUpdateInterval(float InInterval)
-{
-	if (InInterval == TreeUpdateInterval)
-	{
-		return;
-	}
-
-	TreeUpdateInterval = InInterval;
-}
-
-void ACEClonerActor::SetMeshRenderMode(ECEClonerMeshRenderMode InMode)
-{
-	if (InMode == MeshRenderMode)
-	{
-		return;
-	}
-
-	MeshRenderMode = InMode;
-	OnMeshRenderModeChanged();
-}
-
-void ACEClonerActor::SetMeshFacingMode(ENiagaraMeshFacingMode InMode)
-{
-	if (MeshFacingMode == InMode)
-	{
-		return;
-	}
-
-	MeshFacingMode = InMode;
-	OnMeshRendererOptionsChanged();
-}
-
-void ACEClonerActor::SetMeshCastShadows(bool InbCastShadows)
-{
-	if (bMeshCastShadows == InbCastShadows)
-	{
-		return;
-	}
-
-	bMeshCastShadows = InbCastShadows;
-	OnMeshRendererOptionsChanged();
-}
-
-void ACEClonerActor::SetDefaultMeshes(const TArray<TObjectPtr<UStaticMesh>>& InMeshes)
-{
-	DefaultMeshes = InMeshes;
-	OnDefaultMeshesChanged();
-}
-
-void ACEClonerActor::BP_SetDefaultMeshes(const TArray<UStaticMesh*>& InMeshes)
-{
-	DefaultMeshes.Empty(InMeshes.Num());
-	Algo::Transform(InMeshes, DefaultMeshes, [](UStaticMesh* InMesh)->TObjectPtr<UStaticMesh>
-	{
-		return InMesh;
-	});
-	OnDefaultMeshesChanged();
-}
-
-TArray<UStaticMesh*> ACEClonerActor::BP_GetDefaultMeshes() const
-{
-	TArray<UStaticMesh*> Meshes;
-	Meshes.Reserve(DefaultMeshes.Num());
-	Algo::Transform(DefaultMeshes, Meshes, [](const TObjectPtr<UStaticMesh>& InMesh)->UStaticMesh*
-	{
-		return InMesh;
-	});
-	return Meshes;
-}
-
-void ACEClonerActor::SetUseOverrideMaterial(bool bInOverride)
-{
-	if (bUseOverrideMaterial == bInOverride)
-	{
-		return;
-	}
-
-	bUseOverrideMaterial = bInOverride;
-	OnOverrideMaterialChanged();
-}
-
-void ACEClonerActor::SetOverrideMaterial(UMaterialInterface* InMaterial)
-{
-	if (OverrideMaterial == InMaterial)
-	{
-		return;
-	}
-
-	OverrideMaterial = InMaterial;
-	OnOverrideMaterialChanged();
-}
-
-void ACEClonerActor::SetSeed(int32 InSeed)
-{
-	if (InSeed == Seed)
-	{
-		return;
-	}
-
-	Seed = InSeed;
-	OnSeedChanged();
-}
-
-void ACEClonerActor::SetColor(const FLinearColor& InColor)
-{
-	if (InColor == Color)
-	{
-		return;
-	}
-
-	Color = InColor;
-	OnColorChanged();
-}
-
-void ACEClonerActor::SetVisualizeEffectors(bool bInVisualize)
-{
-	if (bVisualizeEffectors == bInVisualize)
-	{
-		return;
-	}
-
-	bVisualizeEffectors = bInVisualize;
-	OnOverrideMaterialChanged();
-}
-
-void ACEClonerActor::SetSurfaceCollisionEnabled(bool bInSurfaceCollisionEnabled)
-{
-	if (bSurfaceCollisionEnabled == bInSurfaceCollisionEnabled)
-	{
-		return;
-	}
-
-	bSurfaceCollisionEnabled = bInSurfaceCollisionEnabled;
-	OnCollisionOptionsChanged();
-}
-
-void ACEClonerActor::SetParticleCollisionEnabled(bool bInParticleCollisionEnabled)
-{
-	if (bParticleCollisionEnabled == bInParticleCollisionEnabled)
-	{
-		return;
-	}
-
-	bParticleCollisionEnabled = bInParticleCollisionEnabled;
-	OnCollisionOptionsChanged();
-}
-
-void ACEClonerActor::SetCollisionVelocityEnabled(bool bInCollisionVelocityEnabled)
-{
-	if (bCollisionVelocityEnabled == bInCollisionVelocityEnabled)
-	{
-		return;
-	}
-
-	bCollisionVelocityEnabled = bInCollisionVelocityEnabled;
-	OnCollisionOptionsChanged();
-}
-
-void ACEClonerActor::SetCollisionIterations(int32 InCollisionIterations)
-{
-	InCollisionIterations = FMath::Max(InCollisionIterations, 1);
-	if (CollisionIterations == InCollisionIterations)
-	{
-		return;
-	}
-
-	CollisionIterations = InCollisionIterations;
-	OnCollisionOptionsChanged();
-}
-
-void ACEClonerActor::SetCollisionGridResolution(int32 InCollisionGridResolution)
-{
-	InCollisionGridResolution = FMath::Max(InCollisionGridResolution, 1);
-	if (CollisionGridResolution == InCollisionGridResolution)
-	{
-		return;
-	}
-
-	CollisionGridResolution = InCollisionGridResolution;
-	OnCollisionOptionsChanged();
-}
-
-void ACEClonerActor::SetCollisionGridSize(const FVector& InCollisionGridSize)
-{
-	const FVector NewCollisionGridSize = InCollisionGridSize.ComponentMax(FVector::ZeroVector);
-	if (CollisionGridSize.Equals(NewCollisionGridSize))
-	{
-		return;
-	}
-
-	CollisionGridSize = NewCollisionGridSize;
-	OnCollisionOptionsChanged();
-}
-
-void ACEClonerActor::SetCollisionRadiusMode(ECEClonerCollisionRadiusMode InMode)
-{
-	if (CollisionRadiusMode == InMode)
-	{
-		return;
-	}
-
-	CollisionRadiusMode = InMode;
-	OnCollisionOptionsChanged();
-}
-
-void ACEClonerActor::SetMassMin(float InMassMin)
-{
-	InMassMin = FMath::Max(InMassMin, 1);
-	if (FMath::IsNearlyEqual(MassMin, InMassMin))
-	{
-		return;
-	}
-
-	MassMin = InMassMin;
-	OnCollisionOptionsChanged();
-}
-
-void ACEClonerActor::SetMassMax(float InMassMax)
-{
-	InMassMax = FMath::Max(InMassMax, 1);
-	if (FMath::IsNearlyEqual(MassMax, InMassMax))
-	{
-		return;
-	}
-
-	MassMax = InMassMax;
-	OnCollisionOptionsChanged();
-}
-
-void ACEClonerActor::SetDeltaStepEnabled(bool bInEnabled)
-{
-	if (bDeltaStepEnabled == bInEnabled)
-	{
-		return;
-	}
-
-	bDeltaStepEnabled = bInEnabled;
-	OnDeltaStepChanged();
-}
-
-void ACEClonerActor::SetDeltaStepRotation(const FRotator& InRotation)
-{
-	if (InRotation == DeltaStepRotation)
-	{
-		return;
-	}
-
-	DeltaStepRotation = InRotation;
-	OnDeltaStepChanged();
-}
-
-void ACEClonerActor::SetDeltaStepScale(const FVector& InScale)
-{
-	if (InScale == DeltaStepScale)
-	{
-		return;
-	}
-
-	DeltaStepScale = InScale;
-	OnDeltaStepChanged();
-}
-
-void ACEClonerActor::SetEnabled(bool bInEnable)
-{
-	if (bInEnable == bEnabled)
-	{
-		return;
-	}
-
-	bEnabled = bInEnable;
-	OnEnabledChanged();
-}
-
-void ACEClonerActor::SetInvertProgress(bool bInInvertProgress)
-{
-	if (bInvertProgress == bInInvertProgress)
-	{
-		return;
-	}
-
-	bInvertProgress = bInInvertProgress;
-	OnProgressChanged();
-}
-
-void ACEClonerActor::SetProgress(float InProgress)
-{
-	if (Progress == InProgress)
-	{
-		return;
-	}
-
-	if (InProgress < 0.f || InProgress > 1.f)
-	{
-		return;
-	}
-
-	Progress = InProgress;
-	OnProgressChanged();
-}
-
-void ACEClonerActor::SetLayoutName(FName InLayoutName)
-{
-	if (LayoutName == InLayoutName)
-	{
-		return;
-	}
-
-	const TArray<FString> LayoutNames = GetClonerLayoutNames();
-	if (!LayoutNames.Contains(InLayoutName))
-	{
-		return;
-	}
-
-	LayoutName = InLayoutName;
-	OnLayoutNameChanged();
-}
-
-void ACEClonerActor::SetRangeEnabled(bool bInRangeEnabled)
-{
-	if (bRangeEnabled == bInRangeEnabled)
-	{
-		return;
-	}
-
-	bRangeEnabled = bInRangeEnabled;
-	OnRangeOptionsChanged();
-}
-
-void ACEClonerActor::SetRangeOffsetMin(const FVector& InRangeOffsetMin)
-{
-	if (RangeOffsetMin == InRangeOffsetMin)
-	{
-		return;
-	}
-
-	RangeOffsetMin = InRangeOffsetMin;
-	OnRangeOptionsChanged();
-}
-
-void ACEClonerActor::SetRangeOffsetMax(const FVector& InRangeOffsetMax)
-{
-	if (RangeOffsetMax == InRangeOffsetMax)
-	{
-		return;
-	}
-
-	RangeOffsetMax = InRangeOffsetMax;
-	OnRangeOptionsChanged();
-}
-
-void ACEClonerActor::SetRangeRotationMin(const FRotator& InRangeRotationMin)
-{
-	if (RangeRotationMin == InRangeRotationMin)
-	{
-		return;
-	}
-
-	RangeRotationMin = InRangeRotationMin;
-	OnRangeOptionsChanged();
-}
-
-void ACEClonerActor::SetRangeRotationMax(const FRotator& InRangeRotationMax)
-{
-	if (RangeRotationMax == InRangeRotationMax)
-	{
-		return;
-	}
-
-	RangeRotationMax = InRangeRotationMax;
-	OnRangeOptionsChanged();
-}
-
-void ACEClonerActor::SetRangeScaleUniform(bool bInRangeScaleUniform)
-{
-	if (bRangeScaleUniform == bInRangeScaleUniform)
-	{
-		return;
-	}
-
-	bRangeScaleUniform = bInRangeScaleUniform;
-	OnRangeOptionsChanged();
-}
-
-void ACEClonerActor::SetRangeScaleMin(const FVector& InRangeScaleMin)
-{
-	if (RangeScaleMin == InRangeScaleMin)
-	{
-		return;
-	}
-
-	RangeScaleMin = InRangeScaleMin;
-	OnRangeOptionsChanged();
-}
-
-void ACEClonerActor::SetRangeScaleMax(const FVector& InRangeScaleMax)
-{
-	if (RangeScaleMax == InRangeScaleMax)
-	{
-		return;
-	}
-
-	RangeScaleMax = InRangeScaleMax;
-	OnRangeOptionsChanged();
-}
-
-void ACEClonerActor::SetRangeScaleUniformMin(float InRangeScaleUniformMin)
-{
-	if (RangeScaleUniformMin == InRangeScaleUniformMin)
-	{
-		return;
-	}
-
-	RangeScaleUniformMin = InRangeScaleUniformMin;
-	OnRangeOptionsChanged();
-}
-
-void ACEClonerActor::SetRangeScaleUniformMax(float InRangeScaleUniformMax)
-{
-	if (RangeScaleUniformMax == InRangeScaleUniformMax)
-	{
-		return;
-	}
-
-	RangeScaleUniformMax = InRangeScaleUniformMax;
-	OnRangeOptionsChanged();
-}
-
-void ACEClonerActor::OnMeshRenderModeChanged()
-{
-	if (!bEnabled)
-	{
-		return;
-	}
-
-	if (const UCEClonerLayoutBase* LayoutSystem = GetActiveLayout())
-	{
-		FNiagaraUserRedirectionParameterStore& ExposedParameters = LayoutSystem->GetSystem()->GetExposedParameters();
-
-		static const FNiagaraVariable MeshModeVar(FNiagaraTypeDefinition(StaticEnum<ECEClonerMeshRenderMode>()), TEXT("MeshRenderMode"));
-		ExposedParameters.SetParameterValue<int32>(static_cast<int32>(MeshRenderMode), MeshModeVar);
-
-		RequestClonerUpdate();
-	}
-}
-
-void ACEClonerActor::OnDeltaStepChanged()
-{
-	if (!ClonerComponent || !bEnabled)
-	{
-		return;
-	}
-
-	if (bDeltaStepEnabled)
-	{
-		ClonerComponent->SetVectorParameter(TEXT("DeltaStepRotation"), FVector(DeltaStepRotation.Roll, DeltaStepRotation.Pitch, DeltaStepRotation.Yaw));
-
-		ClonerComponent->SetVectorParameter(TEXT("DeltaStepScale"), DeltaStepScale);
-	}
-	else
-	{
-		ClonerComponent->SetVectorParameter(TEXT("DeltaStepRotation"), FVector::ZeroVector);
-
-		ClonerComponent->SetVectorParameter(TEXT("DeltaStepScale"), FVector::ZeroVector);
-	}
-
-	RequestClonerUpdate();
-}
-
-void ACEClonerActor::OnSeedChanged()
-{
-	if (!ClonerComponent || !bEnabled)
-	{
-		return;
-	}
-
-	ClonerComponent->SetRandomSeedOffset(Seed);
-
-	RequestClonerUpdate();
-}
-
-void ACEClonerActor::OnColorChanged()
-{
-	if (!ClonerComponent || !bEnabled)
-	{
-		return;
-	}
-
-	ClonerComponent->SetColorParameter(TEXT("EffectorDefaultColor"), Color);
-}
-
-void ACEClonerActor::OnProgressChanged()
-{
-	if (!ClonerComponent || !bEnabled)
-	{
-		return;
-	}
-
-	ClonerComponent->SetFloatParameter(TEXT("ParticleProgress"), Progress * (bInvertProgress ? -1 : 1));
-}
-
-void ACEClonerActor::SetSpawnLoopMode(ECEClonerSpawnLoopMode InMode)
-{
-	if (SpawnLoopMode == InMode)
-	{
-		return;
-	}
-
-	SpawnLoopMode = InMode;
-	OnSpawnOptionsChanged();
-}
-
-void ACEClonerActor::SetSpawnLoopIterations(int32 InIterations)
-{
-	if (SpawnLoopIterations == InIterations)
-	{
-		return;
-	}
-
-	if (InIterations < 1)
-	{
-		return;
-	}
-
-	SpawnLoopIterations = InIterations;
-	OnSpawnOptionsChanged();
-}
-
-void ACEClonerActor::SetSpawnLoopInterval(float InInterval)
-{
-	if (SpawnLoopInterval == InInterval)
-	{
-		return;
-	}
-
-	if (InInterval < 0.f)
-	{
-		return;
-	}
-
-	SpawnLoopInterval = InInterval;
-	OnSpawnOptionsChanged();
-}
-
-void ACEClonerActor::SetSpawnBehaviorMode(ECEClonerSpawnBehaviorMode InMode)
-{
-	if (SpawnBehaviorMode == InMode)
-	{
-		return;
-	}
-
-	SpawnBehaviorMode = InMode;
-	OnSpawnOptionsChanged();
-}
-
-void ACEClonerActor::SetSpawnRate(float InRate)
-{
-	if (SpawnRate == InRate)
-	{
-		return;
-	}
-
-	if (InRate < 0)
-	{
-		return;
-	}
-
-	SpawnRate = InRate;
-	OnSpawnOptionsChanged();
-}
-
-void ACEClonerActor::SetLifetimeEnabled(bool bInEnabled)
-{
-	if (bLifetimeEnabled == bInEnabled)
-	{
-		return;
-	}
-
-	bLifetimeEnabled = bInEnabled;
-	OnLifetimeOptionsChanged();
-}
-
-void ACEClonerActor::SetLifetimeMin(float InMin)
-{
-	if (LifetimeMin == InMin)
-	{
-		return;
-	}
-
-	if (InMin < 0)
-	{
-		return;
-	}
-
-	LifetimeMin = InMin;
-	OnLifetimeOptionsChanged();
-}
-
-void ACEClonerActor::SetLifetimeMax(float InMax)
-{
-	if (LifetimeMax == InMax)
-	{
-		return;
-	}
-
-	if (InMax < 0)
-	{
-		return;
-	}
-
-	LifetimeMax = InMax;
-	OnLifetimeOptionsChanged();
-}
-
-void ACEClonerActor::SetLifetimeScaleEnabled(bool bInEnabled)
-{
-	if (bLifetimeScaleEnabled == bInEnabled)
-	{
-		return;
-	}
-
-	bLifetimeScaleEnabled = bInEnabled;
-	OnLifetimeOptionsChanged();
-}
-
-void ACEClonerActor::SetLifetimeScaleCurve(const FRichCurve& InCurve)
-{
-	LifetimeScaleCurve = InCurve;
-	OnLifetimeOptionsChanged();
-}
-
-UCEClonerLayoutBase* ACEClonerActor::GetActiveLayout() const
-{
-	if (ClonerComponent)
-	{
-		return ClonerComponent->GetClonerActiveLayout();
-	}
-
-	return nullptr;
-}
-
-const FCEClonerEffectorDataInterfaces* ACEClonerActor::GetEffectorDataInterfaces() const
-{
-	if (const UCEClonerLayoutBase* LayoutSystem = GetActiveLayout())
-	{
-		return &LayoutSystem->GetDataInterfaces();
-	}
-
-	return nullptr;
-}
-
-int32 ACEClonerActor::GetMeshCount() const
-{
-	if (ClonerComponent)
-	{
-		if (const UCEClonerLayoutBase* LayoutSystem = GetActiveLayout())
-		{
-			return LayoutSystem->GetMeshRenderer()->Meshes.Num();
-		}
-	}
-	return 0;
-}
-
-bool ACEClonerActor::LinkEffector(ACEEffectorActor* InEffector)
-{
-	if (!IsValid(InEffector) || EffectorsWeak.Contains(InEffector))
-	{
-		return false;
-	}
-
-	EffectorsWeak.Add(InEffector);
-
-	OnEffectorsChanged();
-
-	UE_LOG(LogCEClonerActor, Log, TEXT("%s : Effector %s linked to Cloner"), *GetActorNameOrLabel(), *InEffector->GetActorNameOrLabel());
-
-	return true;
-}
-
-bool ACEClonerActor::UnlinkEffector(ACEEffectorActor* InEffector)
-{
-	if (!InEffector)
-	{
-		return false;
-	}
-
-	if (EffectorsWeak.Remove(InEffector) > 0)
-	{
-		OnEffectorsChanged();
-
-		UE_LOG(LogCEClonerActor, Log, TEXT("%s : Effector %s unlinked from Cloner"), *GetActorNameOrLabel(), *InEffector->GetActorNameOrLabel());
-	}
-
-	return true;
-}
-
-bool ACEClonerActor::IsEffectorLinked(const ACEEffectorActor* InEffector) const
-{
-	return InEffector && EffectorsWeak.Contains(InEffector);
-}
-
-int32 ACEClonerActor::GetEffectorCount() const
-{
-	return EffectorsWeak.Num();
-}
-
-void ACEClonerActor::ForEachEffector(TFunctionRef<bool(ACEEffectorActor*, int32)> InFunction)
-{
-	for (int32 Idx = 0; Idx < EffectorsWeak.Num(); Idx++)
-	{
-		ACEEffectorActor* Effector = EffectorsWeak[Idx].Get();
-		if (!Effector)
-		{
-			continue;
-		}
-
-		if (!InFunction(Effector, Idx))
-		{
-			return;
-		}
-	}
-}
-
-void ACEClonerActor::OnEnabledChanged()
-{
-	if (bEnabled)
-	{
-		OnLayoutNameChanged();
-	}
-	else
-	{
-		ClonerComponent->DeactivateImmediate();
-		ClonerComponent->SetAsset(nullptr);
-	}
-}
-
-void ACEClonerActor::OnLayoutNameChanged()
-{
-	if (!ClonerComponent || !bEnabled)
-	{
-		return;
-	}
-
-	UCEClonerLayoutBase* NewActiveLayout = FindOrAddLayout(LayoutName);
-
-	// Apply layout
-	if (ClonerComponent->SetClonerActiveLayout(NewActiveLayout))
-	{
-		ActiveLayout = ClonerComponent->GetClonerActiveLayout();
-		OnClonerSystemChanged();
-	}
-}
-
-void ACEClonerActor::OnSpawnOptionsChanged()
-{
-	if (!ClonerComponent || !bEnabled)
-	{
-		return;
-	}
-
-	if (const UCEClonerLayoutBase* ActiveSystem = GetActiveLayout())
-	{
-		FNiagaraUserRedirectionParameterStore& ExposedParameters = ActiveSystem->GetSystem()->GetExposedParameters();
-
-		ClonerComponent->SetFloatParameter(TEXT("SpawnLoopInterval"), SpawnLoopInterval);
-
-		ClonerComponent->SetIntParameter(TEXT("SpawnLoopIterations"), SpawnLoopIterations);
-
-		ClonerComponent->SetFloatParameter(TEXT("SpawnRate"), SpawnRate);
-
-		const FNiagaraVariable SpawnBehaviorModeVar(FNiagaraTypeDefinition(StaticEnum<ECEClonerSpawnBehaviorMode>()), TEXT("SpawnBehaviorMode"));
-		ExposedParameters.SetParameterValue<int32>(static_cast<int32>(SpawnBehaviorMode), SpawnBehaviorModeVar);
-
-		const FNiagaraVariable SpawnLoopModeVar(FNiagaraTypeDefinition(StaticEnum<ECEClonerSpawnLoopMode>()), TEXT("SpawnLoopMode"));
-		ExposedParameters.SetParameterValue<int32>(static_cast<int32>(SpawnLoopMode), SpawnLoopModeVar);
-
-		RequestClonerUpdate();
-	}
-}
-
-void ACEClonerActor::OnLifetimeScaleCurveChanged()
-{
-	if (const UNiagaraDataInterfaceCurve* LifetimeCurve = LifetimeScaleCurveDIWeak.Get())
-	{
-		LifetimeScaleCurve = LifetimeCurve->Curve;
-		OnLifetimeOptionsChanged();
-	}
-}
-
-void ACEClonerActor::OnLifetimeOptionsChanged()
-{
-	if (!ClonerComponent || !bEnabled)
-	{
-		return;
-	}
-
-	LifetimeMin = FMath::Max(0, LifetimeMin);
-	LifetimeMax = FMath::Max(LifetimeMin, LifetimeMax);
-
-	ClonerComponent->SetBoolParameter(TEXT("LifetimeEnabled"), bLifetimeEnabled);
-
-	ClonerComponent->SetFloatParameter(TEXT("LifetimeMin"), LifetimeMin);
-
-	ClonerComponent->SetFloatParameter(TEXT("LifetimeMax"), LifetimeMax);
-
-	ClonerComponent->SetBoolParameter(TEXT("LifetimeScaleEnabled"), bLifetimeEnabled && bLifetimeScaleEnabled);
-
-	if (const UCEClonerLayoutBase* ActiveSystem = GetActiveLayout())
-	{
-		const FNiagaraUserRedirectionParameterStore& ExposedParameters = ActiveSystem->GetSystem()->GetExposedParameters();
-
-		static const FNiagaraVariable LifetimeScaleCurveVar(FNiagaraTypeDefinition(UNiagaraDataInterfaceCurve::StaticClass()), TEXT("LifetimeScaleCurve"));
-
-#if WITH_EDITOR
-		if (UNiagaraDataInterfaceCurve* LifetimeCurve = LifetimeScaleCurveDIWeak.Get())
-		{
-			LifetimeCurve->OnChanged().RemoveAll(this);
-		}
-#endif
-
-		if (UNiagaraDataInterfaceCurve* LifetimeScaleCurveDI = Cast<UNiagaraDataInterfaceCurve>(ExposedParameters.GetDataInterface(LifetimeScaleCurveVar)))
-		{
-			LifetimeScaleCurveDIWeak = LifetimeScaleCurveDI;
-			LifetimeScaleCurveDI->Curve = LifetimeScaleCurve;
-
-#if WITH_EDITOR
-			LifetimeScaleCurveDI->UpdateLUT();
-			LifetimeScaleCurveDI->OnChanged().AddUObject(this, &ACEClonerActor::OnLifetimeScaleCurveChanged);
-#endif
-		}
-	}
-
-	RequestClonerUpdate();
-}
-
-void ACEClonerActor::OnCollisionOptionsChanged()
-{
-	if (!ClonerComponent || !bEnabled)
-	{
-		return;
-	}
-
-	MassMin = FMath::Clamp(MassMin, 1, MassMax);
-	MassMax = FMath::Max(MassMax, MassMin);
-
-	ClonerComponent->SetBoolParameter(TEXT("SurfaceCollisionEnabled"), bSurfaceCollisionEnabled);
-	ClonerComponent->SetIntParameter(TEXT("CollisionIterations"), bParticleCollisionEnabled ? CollisionIterations : 0);
-	ClonerComponent->SetBoolParameter(TEXT("CollisionVelocityEnabled"), bParticleCollisionEnabled ? bCollisionVelocityEnabled : false);
-	ClonerComponent->SetIntParameter(TEXT("CollisionGridResolution"), CollisionGridResolution);
-	ClonerComponent->SetVectorParameter(TEXT("CollisionGridSize"), CollisionGridSize);
-	ClonerComponent->SetFloatParameter(TEXT("MassMin"), MassMin);
-	ClonerComponent->SetFloatParameter(TEXT("MassMax"), MassMax);
-
-	// Adjust size based on attached mesh count
-	CollisionRadii.SetNum(GetMeshCount());
-
-	if (const UCEClonerLayoutBase* LayoutSystem = GetActiveLayout())
-	{
-		if (CollisionRadiusMode != ECEClonerCollisionRadiusMode::Manual)
-		{
-			UNiagaraMeshRendererProperties* MeshRenderer = LayoutSystem->GetMeshRenderer();
-
-			for (int32 Idx = 0; Idx < CollisionRadii.Num(); Idx++)
-			{
-				const FNiagaraMeshRendererMeshProperties& MeshProperties = MeshRenderer->Meshes[Idx];
-				FBoxSphereBounds MeshBounds(ForceInitToZero);
-				const FTransform BoundTransform(MeshProperties.Rotation, MeshProperties.PivotOffset, MeshProperties.Scale);
-
-				if (MeshProperties.Mesh)
-				{
-					MeshBounds = MeshProperties.Mesh->GetBounds().TransformBy(BoundTransform);
-				}
-
-				switch (CollisionRadiusMode)
-				{
-					case ECEClonerCollisionRadiusMode::MinExtent:
-						CollisionRadii[Idx] = MeshBounds.BoxExtent.GetMin();
-					break;
-					case ECEClonerCollisionRadiusMode::MaxExtent:
-						CollisionRadii[Idx] = MeshBounds.BoxExtent.GetMax();
-					break;
-					default:
-					case ECEClonerCollisionRadiusMode::ExtentLength:
-						CollisionRadii[Idx] = MeshBounds.SphereRadius;
-					break;
-				}
-			}
-		}
-
-		const FNiagaraUserRedirectionParameterStore& ExposedParameters = LayoutSystem->GetSystem()->GetExposedParameters();
-
-		static const FNiagaraVariable CollisionRadiiVar(FNiagaraTypeDefinition(UNiagaraDataInterfaceArrayFloat::StaticClass()), TEXT("CollisionRadii"));
-
-		if (UNiagaraDataInterfaceArrayFloat* CollisionRadiiDI = Cast<UNiagaraDataInterfaceArrayFloat>(ExposedParameters.GetDataInterface(CollisionRadiiVar)))
-		{
-			CollisionRadiiDI->GetArrayReference() = CollisionRadii;
-		}
-	}
-
-	RequestClonerUpdate();
-}
-
-void ACEClonerActor::OnClonerMeshUpdated(UCEClonerComponent* InClonerComponent)
+struct FCEClonerActorVersion
 {
-	if (InClonerComponent == ClonerComponent)
+	enum Type : int32
 	{
-		OnCollisionOptionsChanged();
+		PreVersioning = 0,
 
-		RequestClonerUpdate();
-	}
-}
-
-void ACEClonerActor::OnClonerSystemChanged()
-{
-	UpdateLayoutOptions();
-}
+		/** Migrating all logic and properties to component */
+		ComponentMigration,
 
-void ACEClonerActor::OnEffectorIdentifierChanged(ACEEffectorActor* InEffector, int32 InOldIdentifier, int32 InNewIdentifier)
-{
-	if (EffectorsWeak.Contains(InEffector))
-	{
-		// Remove effector if it is unregistered from channel
-		if (InNewIdentifier == INDEX_NONE)
-		{
-#if WITH_EDITOR
-			Modify();
-#endif
+		VersionPlusOne,
+		LatestVersion = VersionPlusOne - 1
+	};
 
-			EffectorsWeak.Remove(InEffector);
-		}
+	static constexpr FGuid GUID = FGuid(0x9271D8A2, 0xBF4146B6, 0xA20FC0A3, 0x9D8295B4);
+};
 
-		OnEffectorsChanged();
-	}
-}
+FCustomVersionRegistration GRegisterCEClonerActorVersion(FCEClonerActorVersion::GUID, static_cast<int32>(FCEClonerActorVersion::LatestVersion), TEXT("CEClonerActorVersion"));
 
-void ACEClonerActor::OnEffectorsChanged()
+// Sets default values
+ACEClonerActor::ACEClonerActor()
 {
-	const FCEClonerEffectorDataInterfaces* EffectorDataInterfaces = GetEffectorDataInterfaces();
-
-	if (!EffectorDataInterfaces || !EffectorDataInterfaces->GetIndexArray())
-	{
-		return;
-	}
-
-	// Remove duplicates
-	const TSet<TWeakObjectPtr<ACEEffectorActor>> SetEffectorsWeak(EffectorsWeak);
-	EffectorsWeak = SetEffectorsWeak.Array();
-
-	// Removed effectors
-	for (const TWeakObjectPtr<ACEEffectorActor>& EffectorWeak : EffectorsInternalWeak.Difference(SetEffectorsWeak))
-	{
-		if (ACEEffectorActor* Effector = EffectorWeak.Get())
-		{
-			Effector->OnClonerUnlinked(this);
-		}
-	}
-
-	// Added effectors
-	for (const TWeakObjectPtr<ACEEffectorActor>& EffectorWeak : SetEffectorsWeak.Difference(EffectorsInternalWeak))
-	{
-		if (ACEEffectorActor* Effector = EffectorWeak.Get())
-		{
-			Effector->OnClonerLinked(this);
-		}
-	}
-
-	EffectorsInternalWeak = SetEffectorsWeak;
+	SetCanBeDamaged(false);
+	PrimaryActorTick.bCanEverTick          = false;
 
-	TArray<int32> EffectorIndexes;
-	EffectorIndexes.Reserve(EffectorsWeak.Num());
+	ClonerComponent = CreateDefaultSubobject<UCEClonerComponent>(TEXT("AvaClonerComponent"));
+	SetRootComponent(ClonerComponent);
 
-	for (const TWeakObjectPtr<ACEEffectorActor>& EffectorWeak : EffectorsWeak)
+	if (!IsTemplate())
 	{
-		if (const ACEEffectorActor* Effector = EffectorWeak.Get())
+#if WITH_EDITOR
+		if (GEditor)
 		{
-			const int32 ChannelIdentifier = Effector->GetChannelIdentifier();
-
-			if (ChannelIdentifier != INDEX_NONE)
-			{
-				EffectorIndexes.AddUnique(ChannelIdentifier);
-			}
+			GEditor->GetSelectedActors()->SelectionChangedEvent.AddUObject(this, &ACEClonerActor::OnEditorSelectionChanged);
 		}
-	}
-
-	TArray<int32>& EffectorIndexArray = EffectorDataInterfaces->GetIndexArray()->GetArrayReference();
-	EffectorIndexArray.Empty(EffectorIndexes.Num());
-
-	for (const int32 EffectorIndex : EffectorIndexes)
-	{
-		EffectorIndexArray.Add(EffectorIndex);
-	}
-
-	constexpr bool bImmediateUpdate = false;
-	RequestClonerUpdate(bImmediateUpdate);
-}
-
-void ACEClonerActor::RequestClonerUpdate(bool bInImmediate)
-{
-	if (!bEnabled || !ClonerComponent)
-	{
-		return;
-	}
-
-	if (bInImmediate)
-	{
-		bNeedsRefresh = false;
-		ClonerComponent->RefreshUserParameters();
-	}
-	else
-	{
-		bNeedsRefresh = true;
-	}
-}
+#endif
 
-void ACEClonerActor::OnRangeOptionsChanged()
-{
-	if (!ClonerComponent || !bEnabled)
-	{
-		return;
+		ClonerComponent->OnClonerInitialized().AddUObject(this, &ACEClonerActor::OnClonerInitialized);
 	}
-
-	RangeScaleUniformMin = FMath::Clamp(RangeScaleUniformMin, UE_KINDA_SMALL_NUMBER, RangeScaleUniformMax);
-	RangeScaleUniformMax = FMath::Max3(RangeScaleUniformMin, RangeScaleUniformMax, UE_KINDA_SMALL_NUMBER);
-
-	RangeScaleMin.X = FMath::Clamp(RangeScaleMin.X, UE_KINDA_SMALL_NUMBER, RangeScaleMax.X);
-	RangeScaleMin.Y = FMath::Clamp(RangeScaleMin.Y, UE_KINDA_SMALL_NUMBER, RangeScaleMax.Y);
-	RangeScaleMin.Z = FMath::Clamp(RangeScaleMin.Z, UE_KINDA_SMALL_NUMBER, RangeScaleMax.Z);
-
-	RangeScaleMax.X = FMath::Max3<double>(RangeScaleMin.X, RangeScaleMax.X, UE_KINDA_SMALL_NUMBER);
-	RangeScaleMax.Y = FMath::Max3<double>(RangeScaleMin.Y, RangeScaleMax.Y, UE_KINDA_SMALL_NUMBER);
-	RangeScaleMax.Z = FMath::Max3<double>(RangeScaleMin.Z, RangeScaleMax.Z, UE_KINDA_SMALL_NUMBER);
-
-	ClonerComponent->SetBoolParameter(TEXT("RangeEnabled"), bRangeEnabled);
-
-	ClonerComponent->SetVectorParameter(TEXT("RangeOffsetMin"), RangeOffsetMin);
-
-	ClonerComponent->SetVectorParameter(TEXT("RangeOffsetMax"), RangeOffsetMax);
-
-	ClonerComponent->SetVariableQuat(TEXT("RangeRotationMin"), RangeRotationMin.Quaternion());
-
-	ClonerComponent->SetVariableQuat(TEXT("RangeRotationMax"), RangeRotationMax.Quaternion());
-
-	ClonerComponent->SetBoolParameter(TEXT("RangeScaleUniform"), bRangeScaleUniform);
-
-	ClonerComponent->SetVectorParameter(TEXT("RangeScaleMin"), RangeScaleMin);
-
-	ClonerComponent->SetVectorParameter(TEXT("RangeScaleMax"), RangeScaleMax);
-
-	ClonerComponent->SetFloatParameter(TEXT("RangeScaleUniformMin"), RangeScaleUniformMin);
-
-	ClonerComponent->SetFloatParameter(TEXT("RangeScaleUniformMax"), RangeScaleUniformMax);
-
-	RequestClonerUpdate();
-}
-
-bool ACEClonerActor::IsClonerValid() const
-{
-	return IsValid(ClonerComponent.Get()) && IsValid(ClonerComponent->GetAsset());
 }
 
-void ACEClonerActor::OnMeshRendererOptionsChanged()
+void ACEClonerActor::Serialize(FArchive& InArchive)
 {
-	const UCEClonerLayoutBase* CurrentLayout = ClonerComponent->GetClonerActiveLayout();
+	InArchive.UsingCustomVersion(FCEClonerActorVersion::GUID);
 
-	if (!CurrentLayout)
-	{
-		return;
-	}
+	Super::Serialize(InArchive);
 
-	UNiagaraMeshRendererProperties* MeshRenderer = CurrentLayout->GetMeshRenderer();
+	const int32 Version = InArchive.CustomVer(FCEClonerActorVersion::GUID);
 
-	if (!MeshRenderer)
+	if (Version < FCEClonerActorVersion::ComponentMigration)
 	{
-		return;
+		MigrateToVersion = FCEClonerActorVersion::ComponentMigration;
 	}
-
-	MeshRenderer->FacingMode = MeshFacingMode;
-	MeshRenderer->bCastShadows = bMeshCastShadows;
-
-	RequestClonerUpdate();
 }
 
-void ACEClonerActor::OnDefaultMeshesChanged()
+#if WITH_EDITOR
+FString ACEClonerActor::GetDefaultActorLabel() const
 {
-	// Will force an update
-	TreeUpdateDeltaTime = TreeUpdateInterval;
+	return DefaultLabel;
 }
+#endif
 
-void ACEClonerActor::OnOverrideMaterialChanged()
+void ACEClonerActor::PostLoad()
 {
-	if (ClonerComponent)
-	{
-		UMaterialInterface* OverrideMeshesMaterial = bVisualizeEffectors
-			? LoadObject<UMaterialInterface>(nullptr, UCEClonerEffectorSettings::DefaultMaterialPath)
-			: OverrideMaterial.Get();
+	Super::PostLoad();
 
-		ClonerComponent->SetOverrideMeshesMaterial(OverrideMeshesMaterial);
-		ClonerComponent->SetUseOverrideMeshesMaterial(bUseOverrideMaterial || bVisualizeEffectors);
-	}
+	MigrateDeprecatedProperties();
 }
 
-#if WITH_EDITOR
-void ACEClonerActor::OnReduceMotionGhostingChanged()
+void ACEClonerActor::PostActorCreated()
 {
-	if (UCEClonerSubsystem* ClonerSubsystem = UCEClonerSubsystem::Get())
-	{
-		if (bReduceMotionGhosting)
-		{
-			ClonerSubsystem->EnableNoFlicker();
-		}
-		else
-		{
-			ClonerSubsystem->DisableNoFlicker();
-		}
-	}
-}
+	Super::PostActorCreated();
 
-void ACEClonerActor::OnCVarChanged()
-{
-	if (const UCEClonerSubsystem* ClonerSubsystem = UCEClonerSubsystem::Get())
-	{
-		bReduceMotionGhosting = ClonerSubsystem->IsNoFlickerEnabled();
-	}
+#if WITH_EDITOR
+	bSpawnDefaultActorAttached = true;
+#endif
 }
 
-void ACEClonerActor::ForceUpdateCloner()
+void ACEClonerActor::OnClonerInitialized(UCEClonerComponent* InClonerComponent)
 {
-	if (ClonerComponent)
+#if WITH_EDITOR
+	if (bSpawnDefaultActorAttached)
 	{
-		ClonerComponent->UpdateClonerAttachmentTree();
-		ClonerComponent->UpdateClonerRenderState();
-
-		OnLayoutNameChanged();
-		OnEffectorsChanged();
+		bSpawnDefaultActorAttached = false;
+		ClonerComponent->CreateDefaultActorAttached();
 	}
+#endif
 }
 
-void ACEClonerActor::SpawnLinkedEffector()
+void ACEClonerActor::MigrateDeprecatedProperties()
 {
-	UWorld* ClonerWorld = GetWorld();
-	if (!ClonerWorld)
+	if (MigrateToVersion == INDEX_NONE)
 	{
 		return;
 	}
 
-	FActorSpawnParameters Params;
-	Params.bTemporaryEditorActor = false;
-
-	const FVector ClonerLocation = GetActorLocation();
-	const FRotator ClonerRotation = GetActorRotation();
+	UE_LOG(LogCEClonerActor, Warning, TEXT("%s : Cloner migrating from version %i to latest %i"), *GetActorNameOrLabel(), MigrateToVersion, FCEClonerActorVersion::LatestVersion);
 
-	ACEEffectorActor* EffectorActor = ClonerWorld->SpawnActor<ACEEffectorActor>(ACEEffectorActor::StaticClass(), ClonerLocation, ClonerRotation, Params);
+	PRAGMA_DISABLE_DEPRECATION_WARNINGS
 
-	if (!EffectorActor)
+	if (MigrateToVersion <= FCEClonerActorVersion::ComponentMigration)
 	{
-		return;
-	}
+		ClonerComponent->SetEnabled(bEnabled);
+		ClonerComponent->SetTreeUpdateInterval(TreeUpdateInterval);
+		ClonerComponent->SetSeed(Seed);
+		ClonerComponent->SetLayoutName(LayoutName);
 
-	FActorLabelUtilities::RenameExistingActor(EffectorActor, EffectorActor->GetDefaultActorLabel(), true);
+#if WITH_EDITOR
+		ClonerComponent->SetVisualizerSpriteVisible(bVisualizerSpriteVisible);
+#endif
 
-	// Set default offset for visual feedback
-	EffectorActor->SetOffset(FVector(0, 0, 100));
+		if (UCEClonerMeshRendererExtension* MeshRendererExtension = ClonerComponent->FindOrAddExtension<UCEClonerMeshRendererExtension>())
+		{
+			MeshRendererExtension->SetVisualizeEffectors(bVisualizeEffectors);
+			MeshRendererExtension->SetMeshRenderMode(MeshRenderMode);
+			MeshRendererExtension->SetMeshFacingMode(MeshFacingMode);
+			MeshRendererExtension->SetMeshCastShadows(bMeshCastShadows);
+			MeshRendererExtension->SetDefaultMeshes(DefaultMeshes);
+			MeshRendererExtension->SetUseOverrideMaterial(bUseOverrideMaterial);
+			MeshRendererExtension->SetOverrideMaterial(OverrideMaterial);
+		}
 
-	LinkEffector(EffectorActor);
-}
+		if (UCEClonerCollisionExtension* CollisionExtension = ClonerComponent->FindOrAddExtension<UCEClonerCollisionExtension>())
+		{
+			CollisionExtension->SetSurfaceCollisionEnabled(bSurfaceCollisionEnabled);
+			CollisionExtension->SetParticleCollisionEnabled(bParticleCollisionEnabled);
+			CollisionExtension->SetCollisionVelocityEnabled(bCollisionVelocityEnabled);
+			CollisionExtension->SetCollisionIterations(CollisionIterations);
+			CollisionExtension->SetCollisionGridResolution(CollisionGridResolution);
+			CollisionExtension->SetCollisionGridSize(CollisionGridSize);
+			CollisionExtension->SetCollisionRadiusMode(CollisionRadiusMode);
+			CollisionExtension->SetMassMin(MassMin);
+			CollisionExtension->SetMassMax(MassMax);
+		}
 
-void ACEClonerActor::SpawnDefaultActorAttached()
-{
-	if (!bSpawnDefaultActorAttached)
-	{
-		return;
-	}
+		if (UCEClonerEmitterSpawnExtension* SpawnExtension = ClonerComponent->FindOrAddExtension<UCEClonerEmitterSpawnExtension>())
+		{
+			SpawnExtension->SetSpawnLoopMode(SpawnLoopMode);
+			SpawnExtension->SetSpawnLoopIterations(SpawnLoopIterations);
+			SpawnExtension->SetSpawnLoopIterations(SpawnLoopInterval);
+			SpawnExtension->SetSpawnBehaviorMode(SpawnBehaviorMode);
+			SpawnExtension->SetSpawnRate(SpawnRate);
+		}
 
-	bSpawnDefaultActorAttached = false;
+		if (UCEClonerLifetimeExtension* LifetimeExtension = ClonerComponent->FindOrAddExtension<UCEClonerLifetimeExtension>())
+		{
+			LifetimeExtension->SetLifetimeEnabled(bLifetimeEnabled);
+			LifetimeExtension->SetLifetimeMin(LifetimeMin);
+			LifetimeExtension->SetLifetimeMax(LifetimeMax);
+			LifetimeExtension->SetLifetimeScaleEnabled(bLifetimeScaleEnabled);
+			LifetimeExtension->SetLifetimeScaleCurve(LifetimeScaleCurve);
+		}
 
-	const UCEClonerEffectorSettings* ClonerEffectorSettings = GetDefault<UCEClonerEffectorSettings>();
-	if (!ClonerEffectorSettings || !ClonerEffectorSettings->bSpawnDefaultActorAttached)
-	{
-		return;
-	}
+		if (UCEClonerStepExtension* StepExtension = ClonerComponent->FindOrAddExtension<UCEClonerStepExtension>())
+		{
+			StepExtension->SetDeltaStepEnabled(bDeltaStepEnabled);
+			StepExtension->SetDeltaStepRotation(DeltaStepRotation);
+			StepExtension->SetDeltaStepScale(DeltaStepScale);
+		}
 
-	// Only spawn if world is valid and not a preview actor
-	UWorld* World = GetWorld();
-	if (!World || bIsEditorPreviewActor)
-	{
-		return;
-	}
+		if (UCEClonerRangeExtension* RangeExtension = ClonerComponent->FindOrAddExtension<UCEClonerRangeExtension>())
+		{
+			RangeExtension->SetRangeEnabled(bRangeEnabled);
+			RangeExtension->SetRangeOffsetMin(RangeOffsetMin);
+			RangeExtension->SetRangeOffsetMax(RangeOffsetMax);
+			RangeExtension->SetRangeRotationMin(RangeRotationMin);
+			RangeExtension->SetRangeRotationMax(RangeRotationMax);
+			RangeExtension->SetRangeScaleMin(RangeScaleMin);
+			RangeExtension->SetRangeScaleMax(RangeScaleMax);
+			RangeExtension->SetRangeScaleUniform(bRangeScaleUniform);
+			RangeExtension->SetRangeScaleUniformMin(RangeScaleUniformMin);
+			RangeExtension->SetRangeScaleUniformMax(RangeScaleUniformMax);
+		}
 
-	// Only spawn if no actor is attached below it
-	TArray<AActor*> AttachedActors;
-	constexpr bool bReset = true;
-	constexpr bool bRecursive = false;
-	GetAttachedActors(AttachedActors, bReset, bRecursive);
+		if (UCEClonerProgressExtension* ProgressExtension = ClonerComponent->FindOrAddExtension<UCEClonerProgressExtension>())
+		{
+			ProgressExtension->SetProgress(Progress);
+			ProgressExtension->SetInvertProgress(bInvertProgress);
+		}
 
-	if (!AttachedActors.IsEmpty())
-	{
-		return;
-	}
+		if (UCEClonerEffectorExtension* EffectorExtension = ClonerComponent->FindOrAddExtension<UCEClonerEffectorExtension>())
+		{
+			EffectorExtension->SetColor(Color);
 
-	UStaticMesh* DefaultStaticMesh = ClonerEffectorSettings->DefaultStaticMesh.Get();
-	UMaterialInterface* DefaultMaterial = ClonerEffectorSettings->DefaultMaterial.Get();
+			for (const TWeakObjectPtr<ACEEffectorActor>& EffectorWeak : EffectorsWeak)
+			{
+				if (ACEEffectorActor* Effector = EffectorWeak.Get())
+				{
+					EffectorExtension->LinkEffector(Effector);
+				}
+			}
+		}
 
-	if (!DefaultStaticMesh || !DefaultMaterial)
-	{
-		return;
+		for (const TPair<FName, TObjectPtr<UCEClonerLayoutBase>>& LayoutPair : LayoutInstances)
+		{
+			if (UCEClonerGridLayout* PrevGridLayout = Cast<UCEClonerGridLayout>(LayoutPair.Value))
+			{
+				if (UCEClonerGridLayout* GridLayout = ClonerComponent->FindOrAddLayout<UCEClonerGridLayout>())
+				{
+					GridLayout->SetCountX(PrevGridLayout->GetCountX());
+					GridLayout->SetCountY(PrevGridLayout->GetCountY());
+					GridLayout->SetCountZ(PrevGridLayout->GetCountZ());
+					GridLayout->SetSpacingX(PrevGridLayout->GetSpacingX());
+					GridLayout->SetSpacingY(PrevGridLayout->GetSpacingY());
+					GridLayout->SetSpacingZ(PrevGridLayout->GetSpacingZ());
+					GridLayout->SetConstraint(PrevGridLayout->GetConstraint());
+					GridLayout->SetInvertConstraint(PrevGridLayout->GetInvertConstraint());
+					GridLayout->SetSphereConstraint(PrevGridLayout->GetSphereConstraint());
+					GridLayout->SetCylinderConstraint(PrevGridLayout->GetCylinderConstraint());
+					GridLayout->SetTextureConstraint(PrevGridLayout->GetTextureConstraint());
+					GridLayout->SetTwistAxis(PrevGridLayout->GetTwistAxis());
+					GridLayout->SetTwistFactor(PrevGridLayout->GetTwistFactor());
+				}
+			}
+			else if (UCEClonerLineLayout* PrevLineLayout = Cast<UCEClonerLineLayout>(LayoutPair.Value))
+			{
+				if (UCEClonerLineLayout* LineLayout = ClonerComponent->FindOrAddLayout<UCEClonerLineLayout>())
+				{
+					LineLayout->SetCount(PrevLineLayout->GetCount());
+					LineLayout->SetSpacing(PrevLineLayout->GetSpacing());
+					LineLayout->SetAxis(PrevLineLayout->GetAxis());
+					LineLayout->SetDirection(PrevLineLayout->GetDirection());
+					LineLayout->SetRotation(PrevLineLayout->GetRotation());
+				}
+			}
+			else if (UCEClonerCircleLayout* PrevCircleLayout = Cast<UCEClonerCircleLayout>(LayoutPair.Value))
+			{
+				if (UCEClonerCircleLayout* CircleLayout = ClonerComponent->FindOrAddLayout<UCEClonerCircleLayout>())
+				{
+					CircleLayout->SetCount(PrevCircleLayout->GetCount());
+					CircleLayout->SetRadius(PrevCircleLayout->GetRadius());
+					CircleLayout->SetAngleStart(PrevCircleLayout->GetAngleStart());
+					CircleLayout->SetAngleRatio(PrevCircleLayout->GetAngleRatio());
+					CircleLayout->SetOrientMesh(PrevCircleLayout->GetOrientMesh());
+					CircleLayout->SetPlane(PrevCircleLayout->GetPlane());
+					CircleLayout->SetRotation(PrevCircleLayout->GetRotation());
+					CircleLayout->SetScale(PrevCircleLayout->GetScale());
+				}
+			}
+			else if (UCEClonerCylinderLayout* PrevCylinderLayout = Cast<UCEClonerCylinderLayout>(LayoutPair.Value))
+			{
+				if (UCEClonerCylinderLayout* CylinderLayout = ClonerComponent->FindOrAddLayout<UCEClonerCylinderLayout>())
+				{
+					CylinderLayout->SetBaseCount(PrevCylinderLayout->GetBaseCount());
+					CylinderLayout->SetHeightCount(PrevCylinderLayout->GetHeightCount());
+					CylinderLayout->SetHeight(PrevCylinderLayout->GetHeight());
+					CylinderLayout->SetRadius(PrevCylinderLayout->GetRadius());
+					CylinderLayout->SetAngleStart(PrevCylinderLayout->GetAngleStart());
+					CylinderLayout->SetAngleRatio(PrevCylinderLayout->GetAngleRatio());
+					CylinderLayout->SetOrientMesh(PrevCylinderLayout->GetOrientMesh());
+					CylinderLayout->SetPlane(PrevCylinderLayout->GetPlane());
+					CylinderLayout->SetRotation(PrevCylinderLayout->GetRotation());
+					CylinderLayout->SetScale(PrevCylinderLayout->GetScale());
+				}
+			}
+			else if (UCEClonerSphereUniformLayout* PrevSphereUniformLayout = Cast<UCEClonerSphereUniformLayout>(LayoutPair.Value))
+			{
+				if (UCEClonerSphereUniformLayout* SphereLayout = ClonerComponent->FindOrAddLayout<UCEClonerSphereUniformLayout>())
+				{
+					SphereLayout->SetCount(PrevSphereUniformLayout->GetCount());
+					SphereLayout->SetRadius(PrevSphereUniformLayout->GetRadius());
+					SphereLayout->SetRatio(PrevSphereUniformLayout->GetRatio());
+					SphereLayout->SetOrientMesh(PrevSphereUniformLayout->GetOrientMesh());
+					SphereLayout->SetRotation(PrevSphereUniformLayout->GetRotation());
+					SphereLayout->SetScale(PrevSphereUniformLayout->GetScale());
+				}
+			}
+			else if (UCEClonerHoneycombLayout* PrevHoneycombLayout = Cast<UCEClonerHoneycombLayout>(LayoutPair.Value))
+			{
+				if (UCEClonerHoneycombLayout* HoneycombLayout = ClonerComponent->FindOrAddLayout<UCEClonerHoneycombLayout>())
+				{
+					HoneycombLayout->SetPlane(PrevHoneycombLayout->GetPlane());
+					HoneycombLayout->SetWidthCount(PrevHoneycombLayout->GetWidthCount());
+					HoneycombLayout->SetHeightCount(PrevHoneycombLayout->GetHeightCount());
+					HoneycombLayout->SetWidthOffset(PrevHoneycombLayout->GetWidthOffset());
+					HoneycombLayout->SetHeightOffset(PrevHoneycombLayout->GetHeightOffset());
+					HoneycombLayout->SetHeightSpacing(PrevHoneycombLayout->GetHeightSpacing());
+					HoneycombLayout->SetWidthSpacing(PrevHoneycombLayout->GetWidthSpacing());
+					HoneycombLayout->SetTwistAxis(PrevHoneycombLayout->GetTwistAxis());
+					HoneycombLayout->SetTwistFactor(PrevHoneycombLayout->GetTwistFactor());
+				}
+			}
+			else if (UCEClonerMeshLayout* PrevMeshLayout = Cast<UCEClonerMeshLayout>(LayoutPair.Value))
+			{
+				if (UCEClonerMeshLayout* MeshLayout = ClonerComponent->FindOrAddLayout<UCEClonerMeshLayout>())
+				{
+					MeshLayout->SetCount(PrevMeshLayout->GetCount());
+					MeshLayout->SetAsset(PrevMeshLayout->GetAsset());
+					MeshLayout->SetSampleData(PrevMeshLayout->GetSampleData());
+					MeshLayout->SetSampleActorWeak(PrevMeshLayout->GetSampleActor());
+				}
+			}
+			else if (UCEClonerSplineLayout* PrevSplineLayout = Cast<UCEClonerSplineLayout>(LayoutPair.Value))
+			{
+				if (UCEClonerSplineLayout* SplineLayout = ClonerComponent->FindOrAddLayout<UCEClonerSplineLayout>())
+				{
+					SplineLayout->SetCount(PrevSplineLayout->GetCount());
+					SplineLayout->SetSplineActorWeak(PrevSplineLayout->GetSplineActor());
+					SplineLayout->SetOrientMesh(PrevSplineLayout->GetOrientMesh());
+				}
+			}
+			else if (UCEClonerSphereRandomLayout* PrevSphereRandomLayout = Cast<UCEClonerSphereRandomLayout>(LayoutPair.Value))
+			{
+				if (UCEClonerSphereRandomLayout* SphereLayout = ClonerComponent->FindOrAddLayout<UCEClonerSphereRandomLayout>())
+				{
+					SphereLayout->SetCount(PrevSphereRandomLayout->GetCount());
+					SphereLayout->SetLatitude(PrevSphereRandomLayout->GetLatitude());
+					SphereLayout->SetLongitude(PrevSphereRandomLayout->GetLongitude());
+					SphereLayout->SetDistribution(PrevSphereRandomLayout->GetDistribution());
+					SphereLayout->SetRadius(PrevSphereRandomLayout->GetRadius());
+					SphereLayout->SetOrientMesh(PrevSphereRandomLayout->GetOrientMesh());
+					SphereLayout->SetRotation(PrevSphereRandomLayout->GetRotation());
+					SphereLayout->SetScale(PrevSphereRandomLayout->GetScale());
+				}
+			}
+		}
 	}
 
-	// Spawn attached actor with same flags as this actor
-	FActorSpawnParameters SpawnParameters;
-	SpawnParameters.Owner = this;
-	SpawnParameters.ObjectFlags = GetFlags();
-	SpawnParameters.bTemporaryEditorActor = false;
-
-	if (AStaticMeshActor* DefaultActorAttached = World->SpawnActor<AStaticMeshActor>(GetActorLocation(), GetActorRotation(), SpawnParameters))
-	{
-		UStaticMeshComponent* StaticMeshComponent = DefaultActorAttached->GetStaticMeshComponent();
-		StaticMeshComponent->SetStaticMesh(DefaultStaticMesh);
-		StaticMeshComponent->SetMaterial(0, DefaultMaterial);
-
-		DefaultActorAttached->SetMobility(EComponentMobility::Movable);
-		DefaultActorAttached->AttachToActor(this, FAttachmentTransformRules::KeepWorldTransform);
+	PRAGMA_ENABLE_DEPRECATION_WARNINGS
 
-		FActorLabelUtilities::SetActorLabelUnique(DefaultActorAttached, TEXT("DefaultClone"));
-	}
+	MigrateToVersion = INDEX_NONE;
 }
 
+#if WITH_EDITOR
 void ACEClonerActor::OnEditorSelectionChanged(UObject* InSelection)
 {
 	if (const USelection* ActorSelection = Cast<USelection>(InSelection))
@@ -1519,97 +376,3 @@ void ACEClonerActor::OnEditorSelectionChanged(UObject* InSelection)
 	}
 }
 #endif
-
-void ACEClonerActor::OnVisualizerSpriteVisibleChanged()
-{
-#if WITH_EDITOR
-	UE::ClonerEffector::SetBillboardComponentSprite(this, TEXT("/Script/Engine.Texture2D'/ClonerEffector/Textures/T_ClonerIcon.T_ClonerIcon'"));
-	UE::ClonerEffector::SetBillboardComponentVisibility(this, bVisualizerSpriteVisible);
-#endif
-}
-
-TArray<FString> ACEClonerActor::GetClonerLayoutNames() const
-{
-	TArray<FString> LayoutNamesStrings;
-
-	if (const UCEClonerSubsystem* Subsystem = UCEClonerSubsystem::Get())
-	{
-		const TArray<FName>& LayoutNames = Subsystem->GetLayoutNames();
-		LayoutNamesStrings.Reserve(LayoutNames.Num());
-
-		Algo::Transform(LayoutNames, LayoutNamesStrings, [](const FName& InName)
-		{
-			return InName.ToString();
-		});
-	}
-
-	return LayoutNamesStrings;
-}
-
-UCEClonerLayoutBase* ACEClonerActor::FindOrAddLayout(TSubclassOf<UCEClonerLayoutBase> InClass)
-{
-	const UCEClonerSubsystem* Subsystem = UCEClonerSubsystem::Get();
-
-	if (!Subsystem)
-	{
-		return nullptr;
-	}
-
-	const FName ClassLayoutName = Subsystem->FindLayoutName(InClass);
-
-	if (ClassLayoutName.IsNone())
-	{
-		return nullptr;
-	}
-
-	return FindOrAddLayout(ClassLayoutName);
-}
-
-UCEClonerLayoutBase* ACEClonerActor::FindOrAddLayout(FName InLayoutName)
-{
-	UCEClonerSubsystem* Subsystem = UCEClonerSubsystem::Get();
-	if (!Subsystem)
-	{
-		return nullptr;
-	}
-
-	// Check cached layout instances
-	UCEClonerLayoutBase* NewActiveLayout = nullptr;
-	if (TObjectPtr<UCEClonerLayoutBase> const* LayoutInstance = LayoutInstances.Find(InLayoutName))
-	{
-		NewActiveLayout = *LayoutInstance;
-	}
-
-	// Create new layout instance and cache it
-	if (!NewActiveLayout)
-	{
-		NewActiveLayout = Subsystem->CreateNewLayout(InLayoutName, this);
-		LayoutInstances.Add(InLayoutName, NewActiveLayout);
-	}
-
-	return NewActiveLayout;
-}
-
-void ACEClonerActor::InitializeCloner()
-{
-	if (bClonerInitialized)
-	{
-		return;
-	}
-
-	OnVisualizerSpriteVisibleChanged();
-
-#if WITH_EDITOR
-	OnReduceMotionGhostingChanged();
-
-	if (bSpawnDefaultActorAttached)
-	{
-		SpawnDefaultActorAttached();
-	}
-#endif
-
-	bClonerInitialized = true;
-
-	OnLayoutNameChanged();
-	OnEffectorsChanged();
-}
