MotionDesign : ClonerEffector

diff --git a/Engine/Plugins/Experimental/ClonerEffector/Source/ClonerEffector/Private/Cloner/CEClonerComponent.cpp b/Engine/Plugins/Experimental/ClonerEffector/Source/ClonerEffector/Private/Cloner/CEClonerComponent.cpp
--- a/Engine/Plugins/Experimental/ClonerEffector/Source/ClonerEffector/Private/Cloner/CEClonerComponent.cpp
+++ b/Engine/Plugins/Experimental/ClonerEffector/Source/ClonerEffector/Private/Cloner/CEClonerComponent.cpp
@@ -4,27 +4,26 @@
 
 #include "Async/Async.h"
 #include "Cloner/CEClonerActor.h"
+#include "Cloner/Extensions/CEClonerExtensionBase.h"
 #include "Cloner/Layouts/CEClonerLayoutBase.h"
+#include "Components/BillboardComponent.h"
 #include "Components/BrushComponent.h"
 #include "Components/DynamicMeshComponent.h"
 #include "Components/SkeletalMeshComponent.h"
-#include "DynamicMesh/MeshNormals.h"
-#include "DynamicMesh/MeshTransforms.h"
-#include "DynamicMesh/Operations/MergeCoincidentMeshEdges.h"
-#include "DynamicMeshEditor.h"
-#include "GeometryScript/MeshAssetFunctions.h"
-#include "GeometryScript/MeshBasicEditFunctions.h"
-#include "GeometryScript/SceneUtilityFunctions.h"
+#include "Engine/StaticMeshActor.h"
+#include "Engine/Texture2D.h"
 #include "Materials/Material.h"
+#include "Materials/MaterialInstanceDynamic.h"
 #include "NiagaraMeshRendererProperties.h"
 #include "NiagaraSystem.h"
 #include "ProceduralMeshComponent.h"
-#include "StaticMeshOperations.h"
-#include "UDynamicMesh.h"
+#include "Settings/CEClonerEffectorSettings.h"
 #include "Subsystems/CEClonerSubsystem.h"
+#include "UDynamicMesh.h"
 #include "UObject/Package.h"
 
 #if WITH_EDITOR
+#include "Editor/EditorEngine.h"
 #include "Framework/Notifications/NotificationManager.h"
 #include "Materials/Material.h"
 #include "Misc/MessageDialog.h"
@@ -35,12 +34,14 @@ DEFINE_LOG_CATEGORY_STATIC(LogCEClonerComponent, Log, All);
 
 #define LOCTEXT_NAMESPACE "CEClonerComponent"
 
-UCEClonerComponent::FOnClonerMeshUpdated UCEClonerComponent::OnClonerMeshUpdated;
-UCEClonerComponent::FOnClonerSystemLoaded UCEClonerComponent::OnClonerSystemLoaded;
-
 UCEClonerComponent::UCEClonerComponent()
 	: UNiagaraComponent()
 {
+	PrimaryComponentTick.bCanEverTick          = true;
+	PrimaryComponentTick.bStartWithTickEnabled = true;
+	PrimaryComponentTick.bTickEvenWhenPaused   = true;
+	PrimaryComponentTick.bHighPriority         = true;
+
 	CastShadow = true;
 	bReceivesDecals = true;
 	bAutoActivate = true;
@@ -68,6 +69,39 @@ UCEClonerComponent::UCEClonerComponent()
 		FCoreUObjectDelegates::OnObjectPropertyChanged.RemoveAll(this);
 		FCoreUObjectDelegates::OnObjectPropertyChanged.AddUObject(this, &UCEClonerComponent::OnActorPropertyChanged);
 #endif
+
+		// Apply default layout
+		const TArray<FString> LayoutNames = GetClonerLayoutNames();
+		LayoutName = !LayoutNames.IsEmpty() ? FName(LayoutNames[0]) : NAME_None;
+	}
+}
+
+void UCEClonerComponent::PostInitProperties()
+{
+	Super::PostInitProperties();
+
+	if (HasAnyFlags(RF_ClassDefaultObject))
+	{
+		// Register new type def for niagara
+
+		constexpr ENiagaraTypeRegistryFlags MeshFlags =
+			ENiagaraTypeRegistryFlags::AllowAnyVariable |
+			ENiagaraTypeRegistryFlags::AllowParameter;
+
+		FNiagaraTypeRegistry::Register(FNiagaraTypeDefinition(StaticEnum<ECEClonerMeshRenderMode>()), MeshFlags);
+		FNiagaraTypeRegistry::Register(FNiagaraTypeDefinition(StaticEnum<ECEClonerGridConstraint>()), MeshFlags);
+		FNiagaraTypeRegistry::Register(FNiagaraTypeDefinition(StaticEnum<ECEClonerPlane>()), MeshFlags);
+		FNiagaraTypeRegistry::Register(FNiagaraTypeDefinition(StaticEnum<ECEClonerAxis>()), MeshFlags);
+		FNiagaraTypeRegistry::Register(FNiagaraTypeDefinition(StaticEnum<ECEClonerEasing>()), MeshFlags);
+		FNiagaraTypeRegistry::Register(FNiagaraTypeDefinition(StaticEnum<ECEClonerMeshAsset>()), MeshFlags);
+		FNiagaraTypeRegistry::Register(FNiagaraTypeDefinition(StaticEnum<ECEClonerMeshSampleData>()), MeshFlags);
+		FNiagaraTypeRegistry::Register(FNiagaraTypeDefinition(StaticEnum<ECEClonerEffectorType>()), MeshFlags);
+		FNiagaraTypeRegistry::Register(FNiagaraTypeDefinition(StaticEnum<ECEClonerTextureSampleChannel>()), MeshFlags);
+		FNiagaraTypeRegistry::Register(FNiagaraTypeDefinition(StaticEnum<ECEClonerCompareMode>()), MeshFlags);
+		FNiagaraTypeRegistry::Register(FNiagaraTypeDefinition(StaticEnum<ECEClonerEffectorMode>()), MeshFlags);
+		FNiagaraTypeRegistry::Register(FNiagaraTypeDefinition(StaticEnum<ECEClonerSpawnLoopMode>()), MeshFlags);
+		FNiagaraTypeRegistry::Register(FNiagaraTypeDefinition(StaticEnum<ECEClonerSpawnBehaviorMode>()), MeshFlags);
+		FNiagaraTypeRegistry::Register(FNiagaraTypeDefinition(StaticEnum<ECEClonerEffectorPushDirection>()), MeshFlags);
 	}
 }
 
@@ -76,7 +110,34 @@ void UCEClonerComponent::PostLoad()
 	SetAsset(nullptr);
 
 	Super::PostLoad();
+
+	InitializeCloner();
+}
+
+#if WITH_EDITOR
+void UCEClonerComponent::PostEditUndo()
+{
+	Super::PostEditUndo();
+
+	ForceUpdateCloner();
+}
+
+const TCEPropertyChangeDispatcher<UCEClonerComponent> UCEClonerComponent::PropertyChangeDispatcher =
+{
+	{ GET_MEMBER_NAME_CHECKED(UCEClonerComponent, bEnabled), &UCEClonerComponent::OnEnabledChanged },
+	{ GET_MEMBER_NAME_CHECKED(UCEClonerComponent, Seed), &UCEClonerComponent::OnSeedChanged },
+	/** Layout */
+	{ GET_MEMBER_NAME_CHECKED(UCEClonerComponent, LayoutName), &UCEClonerComponent::OnLayoutNameChanged },
+	{ GET_MEMBER_NAME_CHECKED(UCEClonerComponent, bVisualizerSpriteVisible), &UCEClonerComponent::OnVisualizerSpriteVisibleChanged },
+};
+
+void UCEClonerComponent::PostEditChangeProperty(FPropertyChangedEvent& InPropertyChangedEvent)
+{
+	Super::PostEditChangeProperty(InPropertyChangedEvent);
+
+	PropertyChangeDispatcher.OnPropertyChanged(this, InPropertyChangedEvent);
 }
+#endif
 
 void UCEClonerComponent::UpdateClonerRenderState()
 {
@@ -95,7 +156,11 @@ void UCEClonerComponent::UpdateClonerRenderState()
 		return;
 	}
 
+#if WITH_EDITOR
 	UpdateDirtyMeshesAsync();
+#else
+	OnDirtyMeshesUpdated(true);
+#endif
 }
 
 void UCEClonerComponent::RefreshUserParameters() const
@@ -107,35 +172,9 @@ void UCEClonerComponent::RefreshUserParameters() const
 	}
 }
 
-void UCEClonerComponent::OnRenderStateDirty(UActorComponent& InComponent)
-{
-	const AActor* OwningActor = InComponent.GetOwner();
-	const ACEClonerActor* ClonerActor = GetOuterACEClonerActor();
-
-	if (!IsValid(OwningActor) || !IsValid(ClonerActor))
-	{
-		return;
-	}
-
-	if (OwningActor->GetWorld() != ClonerActor->GetWorld())
-	{
-		return;
-	}
-
-	/* Owning actor is updated and is attached to cloner
-		or owning actor is detached from cloner */
-	if (OwningActor->IsAttachedTo(ClonerActor)
-		|| ClonerTree.ItemAttachmentMap.Contains(OwningActor))
-	{
-		UE_LOG(LogCEClonerComponent, Log, TEXT("%s : Detected attachment change for %s"), *ClonerActor->GetActorNameOrLabel(), *OwningActor->GetActorNameOrLabel());
-
-		// Only update attachment tree
-		UpdateClonerAttachmentTree();
-	}
-}
-
 void UCEClonerComponent::UpdateClonerAttachmentTree(bool bInReset)
 {
+#if WITH_EDITOR
 	if (ClonerTree.Status == ECEClonerAttachmentStatus::Updated)
 	{
 		ClonerTree.Status = ECEClonerAttachmentStatus::Outdated;
@@ -148,6 +187,7 @@ void UCEClonerComponent::UpdateClonerAttachmentTree(bool bInReset)
 	}
 
 	UpdateAttachmentTree();
+#endif
 }
 
 void UCEClonerComponent::UpdateAttachmentTree()
@@ -174,6 +214,11 @@ void UCEClonerComponent::UpdateAttachmentTree()
 	{
 		AActor* RootChild = RootChildren[RootIdx];
 
+		if (!RootChild)
+		{
+			continue;
+		}
+
 		UpdateActorAttachment(RootChild, nullptr);
 
 		// Lets find the old root idx
@@ -181,7 +226,7 @@ void UCEClonerComponent::UpdateAttachmentTree()
 		TObjectPtr<UStaticMesh> CombineBakedMesh = nullptr;
 		if (OldIdx != INDEX_NONE)
 		{
-			CombineBakedMesh = ClonerTree.MergedBakedMeshes[OldIdx];
+			CombineBakedMesh = ClonerTree.MergedBakedMeshes[OldIdx].Get();
 
 			// Did we rearrange stuff ?
 			if (RootIdx != OldIdx)
@@ -221,7 +266,7 @@ void UCEClonerComponent::UpdateAttachmentTree()
 		{
 			if (ClonedItem->MeshStatus == ECEClonerAttachmentStatus::Outdated)
 			{
-				DirtyItemAttachments.Add(ClonedItem);
+				ClonerTree.DirtyItemAttachments.Add(ClonedItem);
 				InvalidateBakedStaticMesh(ClonedActor.Get());
 			}
 			bClonerMeshesDirty = true;
@@ -235,37 +280,11 @@ void UCEClonerComponent::UpdateAttachmentTree()
 		bClonerMeshesDirty = true;
 	}
 
-	if (!DirtyItemAttachments.IsEmpty())
+	if (!ClonerTree.DirtyItemAttachments.IsEmpty())
 	{
 		bClonerMeshesDirty = true;
 	}
 
-	// Use default meshes if nothing is attached
-	if (NewRootActors.IsEmpty() && NewCombinesMeshes.IsEmpty())
-	{
-		if (const ACEClonerActor* ClonerActor = GetOuterACEClonerActor())
-		{
-			const TArray<TObjectPtr<UStaticMesh>>& DefaultMeshes = ClonerActor->GetDefaultMeshes();
-
-			if (ClonerTree.MergedBakedMeshes.Num() != DefaultMeshes.Num())
-			{
-				bClonerMeshesDirty = true;
-			}
-
-			for (const TObjectPtr<UStaticMesh>& DefaultMesh : DefaultMeshes)
-			{
-				const int32 Idx = NewCombinesMeshes.Add(DefaultMesh);
-
-				// Only update when there is a change
-				if (!ClonerTree.MergedBakedMeshes.IsValidIndex(Idx)
-					|| DefaultMesh != ClonerTree.MergedBakedMeshes[Idx])
-				{
-					bClonerMeshesDirty = true;
-				}
-			}
-		}
-	}
-
 	ClonerTree.RootActors = NewRootActors;
 	ClonerTree.MergedBakedMeshes = NewCombinesMeshes;
 	ClonerTree.Status = ECEClonerAttachmentStatus::Updated;
@@ -593,7 +612,7 @@ void UCEClonerComponent::OnMeshChanged(UStaticMeshComponent*, AActor* InActor)
 
 		Item->MeshStatus = ECEClonerAttachmentStatus::Outdated;
         InvalidateBakedStaticMesh(InActor);
-		DirtyItemAttachments.Add(Item);
+		ClonerTree.DirtyItemAttachments.Add(Item);
 	}
 }
 
@@ -674,8 +693,8 @@ void UCEClonerComponent::UpdateDirtyMeshesAsync()
 	}
 
 	bClonerMeshesUpdating = true;
-	TSet<FCEClonerAttachmentItem*> UpdateItems = DirtyItemAttachments;
-	DirtyItemAttachments.Empty();
+	TSet<FCEClonerAttachmentItem*> UpdateItems = ClonerTree.DirtyItemAttachments;
+	ClonerTree.DirtyItemAttachments.Empty();
 
 	// Update baked dynamic meshes on other thread
 	TWeakObjectPtr<UCEClonerComponent> ThisWeak(this);
@@ -700,7 +719,7 @@ void UCEClonerComponent::UpdateDirtyMeshesAsync()
 			if (IsGarbageCollectingAndLockingUObjectHashTables())
 			{
 				bSuccess = false;
-				This->DirtyItemAttachments.Add(Item);
+				This->ClonerTree.DirtyItemAttachments.Add(Item);
 				continue;
 			}
 
@@ -765,16 +784,16 @@ void UCEClonerComponent::UpdateActorBakedDynamicMesh(AActor* InActor)
 		return;
 	}
 
-	FCEClonerAttachmentItem* AttachmentItem = ClonerTree.ItemAttachmentMap.Find(InActor);
+	const AActor* ClonerActor = GetOwner();
 
-	if (!AttachmentItem || AttachmentItem->MeshStatus != ECEClonerAttachmentStatus::Outdated)
+	if (!ClonerActor)
 	{
 		return;
 	}
 
-	AActor* ClonerActor = GetOwner();
+	FCEClonerAttachmentItem* AttachmentItem = ClonerTree.ItemAttachmentMap.Find(InActor);
 
-	if (!ClonerActor)
+	if (!AttachmentItem || AttachmentItem->MeshStatus != ECEClonerAttachmentStatus::Outdated)
 	{
 		return;
 	}
@@ -783,285 +802,15 @@ void UCEClonerComponent::UpdateActorBakedDynamicMesh(AActor* InActor)
 
 	UE_LOG(LogCEClonerComponent, Log, TEXT("%s : Updating baked actor mesh %s"), *ClonerActor->GetActorNameOrLabel(), *InActor->GetActorNameOrLabel());
 
-	using namespace UE::Geometry;
-
-	TArray<FDynamicMesh3> ConvertedMeshes;
-	TArray<TWeakObjectPtr<UMaterialInterface>> MeshesMaterials;
-
-	const FTransform& SourceTransform = InActor->GetActorTransform();
-	UDynamicMesh* OutputDynamicMesh = NewObject<UDynamicMesh>();
-
-	static FGeometryScriptCopyMeshFromAssetOptions OutputMeshOptions;
-	OutputMeshOptions.bIgnoreRemoveDegenerates = false;
-	OutputMeshOptions.bRequestTangents = false;
-	OutputMeshOptions.bApplyBuildSettings = false;
-
-	// Dynamic mesh components
-	{
-		TArray<UDynamicMeshComponent*> Components;
-		InActor->GetComponents(Components, false);
-		for (UDynamicMeshComponent* Component : Components)
-		{
-			const UDynamicMesh* DynamicMesh = Component->GetDynamicMesh();
-			// Transform the new mesh relative to the component
-			const FTransform RelativeTransform = Component->GetComponentTransform().GetRelativeTransform(SourceTransform);
-			MeshesMaterials.Append(Component->GetMaterials());
-			// Create a copy
-			DynamicMesh->ProcessMesh([&ConvertedMeshes, RelativeTransform](const FDynamicMesh3& EditMesh)
-			{
-				FDynamicMesh3 CopyMesh = EditMesh;
-				MeshTransforms::ApplyTransform(CopyMesh, RelativeTransform);
-				ConvertedMeshes.Add(MoveTemp(CopyMesh));
-			});
-		}
-	}
-
-	// Skeletal mesh components
-	{
-		static FGeometryScriptMeshReadLOD SkeletalMeshLOD;
-		SkeletalMeshLOD.LODType = EGeometryScriptLODType::SourceModel;
-
-		TArray<USkeletalMeshComponent*> Components;
-		InActor->GetComponents(Components, false);
-		for (USkeletalMeshComponent* Component : Components)
-		{
-			USkeletalMesh* SkeletalMesh = Component->GetSkeletalMeshAsset();
-
-			// convert to dynamic mesh
-			EGeometryScriptOutcomePins OutResult;
-			UGeometryScriptLibrary_StaticMeshFunctions::CopyMeshFromSkeletalMesh(SkeletalMesh, OutputDynamicMesh, OutputMeshOptions, SkeletalMeshLOD, OutResult);
-			if (OutResult == EGeometryScriptOutcomePins::Success)
-			{
-				// Transform the new mesh relative to the component
-				const FTransform RelativeTransform = Component->GetComponentTransform().GetRelativeTransform(SourceTransform);
-				MeshesMaterials.Append(Component->GetMaterials());
-				OutputDynamicMesh->EditMesh([&ConvertedMeshes, RelativeTransform](FDynamicMesh3& EditMesh)
-				{
-					MeshTransforms::ApplyTransform(EditMesh, RelativeTransform);
-					ConvertedMeshes.Add(MoveTemp(EditMesh));
-					// replace by empty mesh
-					FDynamicMesh3 EmptyMesh;
-					EditMesh = MoveTemp(EmptyMesh);
-				}, EDynamicMeshChangeType::GeneralEdit, EDynamicMeshAttributeChangeFlags::Unknown, true);
-			}
-		}
-	}
-
-	// Brush components
-	{
-		static const FGeometryScriptCopyMeshFromComponentOptions Options;
-		static FTransform Transform;
-
-		TArray<UBrushComponent*> Components;
-		InActor->GetComponents(Components, false);
-		for (UBrushComponent* Component : Components)
-		{
-			// convert to dynamic mesh
-			EGeometryScriptOutcomePins OutResult;
-			UGeometryScriptLibrary_SceneUtilityFunctions::CopyMeshFromComponent(Component, OutputDynamicMesh, Options, false, Transform, OutResult);
-			if (OutResult == EGeometryScriptOutcomePins::Success)
-			{
-				// Transform the new mesh relative to the component
-				const FTransform RelativeTransform = Component->GetComponentTransform().GetRelativeTransform(SourceTransform);
-				for (int32 MatIdx = 0; MatIdx < Component->GetNumMaterials(); MatIdx++)
-				{
-					MeshesMaterials.Add(Component->GetMaterial(MatIdx));
-				}
-				OutputDynamicMesh->EditMesh([&ConvertedMeshes, RelativeTransform](FDynamicMesh3& EditMesh)
-				{
-					MeshTransforms::ApplyTransform(EditMesh, RelativeTransform);
-					ConvertedMeshes.Add(MoveTemp(EditMesh));
-					// replace by empty mesh
-					FDynamicMesh3 EmptyMesh;
-					EditMesh = MoveTemp(EmptyMesh);
-				}, EDynamicMeshChangeType::GeneralEdit, EDynamicMeshAttributeChangeFlags::Unknown, true);
-			}
-		}
-	}
-
-	// Procedural mesh components
-	{
-		TArray<UProceduralMeshComponent*> Components;
-		InActor->GetComponents(Components, false);
-		for (UProceduralMeshComponent* Component : Components)
-		{
-			const int32 SectionCount = Component->GetNumSections();
-			if (SectionCount == 0)
-			{
-				continue;
-			}
-
-			// Transform the new mesh relative to the component
-			const FTransform RelativeTransform = Component->GetComponentTransform().GetRelativeTransform(SourceTransform);
-
-			FDynamicMesh3 Mesh;
-			Mesh.EnableAttributes();
-			Mesh.Attributes()->EnablePrimaryColors();
-			Mesh.Attributes()->EnableMaterialID();
-			Mesh.Attributes()->SetNumNormalLayers(1);
-			Mesh.Attributes()->SetNumUVLayers(1);
-			Mesh.Attributes()->SetNumPolygroupLayers(1);
-			Mesh.Attributes()->EnableTangents();
-
-			FDynamicMeshColorOverlay* ColorOverlay = Mesh.Attributes()->PrimaryColors();
-			FDynamicMeshNormalOverlay* NormalOverlay = Mesh.Attributes()->PrimaryNormals();
-			FDynamicMeshUVOverlay* UVOverlay = Mesh.Attributes()->PrimaryUV();
-			FDynamicMeshMaterialAttribute* MaterialAttr = Mesh.Attributes()->GetMaterialID();
-			FDynamicMeshPolygroupAttribute* GroupAttr = Mesh.Attributes()->GetPolygroupLayer(0);
-			FDynamicMeshNormalOverlay* TangentOverlay = Mesh.Attributes()->PrimaryTangents();
-
-			for (int32 SectionIdx = 0; SectionIdx < SectionCount; SectionIdx++)
-			{
-				if (FProcMeshSection* Section = Component->GetProcMeshSection(SectionIdx))
-				{
-					if (Section->bSectionVisible)
-					{
-						TArray<int32> VtxIds;
-						TArray<int32> NormalIds;
-						TArray<int32> ColorIds;
-						TArray<int32> UVIds;
-						TArray<int32> TaIds;
-
-						// copy vertices data (position, normal, color, UV, tangent)
-						for (FProcMeshVertex& SectionVertex : Section->ProcVertexBuffer)
-						{
-							int32 VId = Mesh.AppendVertex(SectionVertex.Position);
-							VtxIds.Add(VId);
-
-							int32 NId = NormalOverlay->AppendElement(static_cast<FVector3f>(SectionVertex.Normal));
-							NormalIds.Add(NId);
-
-							int32 CId = ColorOverlay->AppendElement(static_cast<FVector4f>(SectionVertex.Color));
-							ColorIds.Add(CId);
-
-							int32 UVId = UVOverlay->AppendElement(static_cast<FVector2f>(SectionVertex.UV0));
-							UVIds.Add(UVId);
-
-							int32 TaId = TangentOverlay->AppendElement(static_cast<FVector3f>(SectionVertex.Tangent.TangentX));
-							TaIds.Add(TaId);
-						}
-
-						// copy tris data
-						if (Section->ProcIndexBuffer.Num() % 3 != 0)
-						{
-							continue;
-						}
-						for (int32 Idx = 0; Idx < Section->ProcIndexBuffer.Num(); Idx+=3)
-						{
-							int32 VIdx1 = Section->ProcIndexBuffer[Idx];
-							int32 VIdx2 = Section->ProcIndexBuffer[Idx + 1];
-							int32 VIdx3 = Section->ProcIndexBuffer[Idx + 2];
-
-							int32 VId1 = VtxIds[VIdx1];
-							int32 VId2 = VtxIds[VIdx2];
-							int32 VId3 = VtxIds[VIdx3];
-
-							int32 TId = Mesh.AppendTriangle(VId1, VId2, VId3, SectionIdx);
-							if (TId < 0)
-							{
-								continue;
-							}
-
-							NormalOverlay->SetTriangle(TId, FIndex3i(NormalIds[VIdx1], NormalIds[VIdx2], NormalIds[VIdx3]), true);
-							ColorOverlay->SetTriangle(TId, FIndex3i(ColorIds[VIdx1], ColorIds[VIdx2], ColorIds[VIdx3]), true);
-							UVOverlay->SetTriangle(TId, FIndex3i(UVIds[VIdx1], UVIds[VIdx2], UVIds[VIdx3]), true);
-							TangentOverlay->SetTriangle(TId, FIndex3i(TaIds[VIdx1], TaIds[VIdx2], TaIds[VIdx3]), true);
-
-							MaterialAttr->SetValue(TId, SectionIdx);
-							GroupAttr->SetValue(TId, SectionIdx);
-						}
-					}
-				}
-			}
-
-			if (Mesh.TriangleCount() > 0)
-			{
-				MeshTransforms::ApplyTransform(Mesh, RelativeTransform);
-				MeshesMaterials.Append(Component->GetMaterials());
-				ConvertedMeshes.Add(MoveTemp(Mesh));
-			}
-		}
-	}
-
-	// Static mesh components
-	{
-		static FGeometryScriptMeshReadLOD StaticMeshLOD;
-		StaticMeshLOD.LODType = EGeometryScriptLODType::RenderData;
-
-		TArray<UStaticMeshComponent*> Components;
-		InActor->GetComponents(Components, false);
-		for (UStaticMeshComponent* Component : Components)
-		{
-			UStaticMesh* StaticMesh = Component->GetStaticMesh();
-			// convert to dynamic mesh
-			EGeometryScriptOutcomePins OutResult;
-			UGeometryScriptLibrary_StaticMeshFunctions::CopyMeshFromStaticMesh(StaticMesh, OutputDynamicMesh, OutputMeshOptions, StaticMeshLOD, OutResult);
-			if (OutResult == EGeometryScriptOutcomePins::Success)
-			{
-				// Transform the new mesh relative to the component
-				const FTransform RelativeTransform = Component->GetComponentTransform().GetRelativeTransform(SourceTransform);
-				MeshesMaterials.Append(Component->GetMaterials());
-				OutputDynamicMesh->EditMesh([&ConvertedMeshes, RelativeTransform](FDynamicMesh3& EditMesh)
-				{
-					MeshTransforms::ApplyTransform(EditMesh, RelativeTransform);
-					ConvertedMeshes.Add(MoveTemp(EditMesh));
-					// replace by empty mesh
-					FDynamicMesh3 EmptyMesh;
-					EditMesh = MoveTemp(EmptyMesh);
-				}, EDynamicMeshChangeType::GeneralEdit, EDynamicMeshAttributeChangeFlags::Unknown, true);
-			}
-		}
-	}
-
-	// Lets combine all meshes components from this actor together
-	OutputDynamicMesh->EditMesh([&ConvertedMeshes](FDynamicMesh3& InMergedMesh)
-	{
-		InMergedMesh.Clear();
-		InMergedMesh.EnableAttributes();
-		InMergedMesh.Attributes()->SetNumNormalLayers(1);
-
-		FDynamicMeshEditor Editor(&InMergedMesh);
-		FGeometryScriptAppendMeshOptions AppendOptions;
-		AppendOptions.CombineMode = EGeometryScriptCombineAttributesMode::EnableAllMatching;
-		int32 MaterialCount = 0;
+	UDynamicMesh* Mesh = NewObject<UDynamicMesh>();
+	TArray<TWeakObjectPtr<UMaterialInterface>> MeshMaterials;
 
-		// Convert meshes
-		for (int32 MeshIdx = 0; MeshIdx < ConvertedMeshes.Num(); MeshIdx++)
-		{
-			const FDynamicMesh3& ConvertedMesh = ConvertedMeshes[MeshIdx];
-
-			// Enable matching attributes & append mesh
-			FMeshIndexMappings TmpMappings;
-			AppendOptions.UpdateAttributesForCombineMode(InMergedMesh, ConvertedMesh);
-			Editor.AppendMesh(&ConvertedMesh, TmpMappings);
-
-			// Fix triangles materials linking
-			if (ConvertedMesh.HasAttributes() && ConvertedMesh.Attributes()->HasMaterialID())
-			{
-				const FDynamicMeshMaterialAttribute* FromMaterialIDAttrib = ConvertedMesh.Attributes()->GetMaterialID();
-				FDynamicMeshMaterialAttribute* ToMaterialIDAttrib = InMergedMesh.Attributes()->GetMaterialID();
-				TMap<int32, int32> MaterialMap;
-				for (const TPair<int32, int32>& FromToTId : TmpMappings.GetTriangleMap().GetForwardMap())
-				{
-					const int32 FromMatId = FromMaterialIDAttrib->GetValue(FromToTId.Key);
-					const int32 ToMatId = FromMatId + MaterialCount;
-					MaterialMap.Add(FromMatId, ToMatId);
-					ToMaterialIDAttrib->SetNewValue(FromToTId.Value, ToMatId);
-				}
-				MaterialCount += MaterialMap.Num();
-			}
-		}
+	MeshBuilder.AppendActor(InActor);
+	MeshBuilder.BuildDynamicMesh(Mesh, MeshMaterials);
 
-		if (InMergedMesh.TriangleCount() > 0)
-		{
-			// Merge shared edges
-			FMergeCoincidentMeshEdges WeldOp(&InMergedMesh);
-			WeldOp.Apply();
-		}
-	}, EDynamicMeshChangeType::GeneralEdit, EDynamicMeshAttributeChangeFlags::Unknown, true);
+	AttachmentItem->BakedMesh = Mesh;
+	AttachmentItem->BakedMaterials = MoveTemp(MeshMaterials);
 
-	AttachmentItem->BakedMesh = OutputDynamicMesh;
-	AttachmentItem->BakedMaterials = MeshesMaterials;
 	// Was the mesh invalidated during the update process ?
 	AttachmentItem->MeshStatus = AttachmentItem->MeshStatus == ECEClonerAttachmentStatus::Outdated ? ECEClonerAttachmentStatus::Outdated : ECEClonerAttachmentStatus::Updated;
 	InvalidateBakedStaticMesh(AttachmentItem->ItemActor.Get());
@@ -1074,189 +823,95 @@ void UCEClonerComponent::UpdateRootActorBakedStaticMesh(AActor* InRootActor)
 		return;
 	}
 
-	const int32 RootIdx = ClonerTree.RootActors.Find(InRootActor);
+	const AActor* ClonerActor = GetOwner();
 
-	if (RootIdx == INDEX_NONE)
+	if (!ClonerActor)
 	{
 		return;
 	}
 
-	const FCEClonerAttachmentItem* RootAttachmentItem = ClonerTree.ItemAttachmentMap.Find(InRootActor);
+	const int32 RootIdx = ClonerTree.RootActors.Find(InRootActor);
 
-	if (!RootAttachmentItem)
+	if (RootIdx == INDEX_NONE)
 	{
 		return;
 	}
 
-	const AActor* ClonerActor = GetOwner();
+	const FCEClonerAttachmentItem* RootAttachmentItem = ClonerTree.ItemAttachmentMap.Find(InRootActor);
 
-	if (!ClonerActor)
+	if (!RootAttachmentItem)
 	{
 		return;
 	}
 
 	UE_LOG(LogCEClonerComponent, Log, TEXT("%s : Updating root merged baked mesh %s"), *ClonerActor->GetActorNameOrLabel(), *InRootActor->GetActorNameOrLabel());
 
-	using namespace UE::Geometry;
-
-#if WITH_EDITOR
-	TSet<int32> MaterialsMissingNiagaraUsageFlag;
-#endif
-
 	TArray<FCEClonerAttachmentItem*> AttachmentItems;
 	GetActorAttachmentItems(InRootActor, AttachmentItems);
 
-	// Lets combine all dynamic meshes from item attachments and self into one mesh
-	UDynamicMesh* MergedAttachmentMesh = NewObject<UDynamicMesh>();
-	TArray<TObjectPtr<UMaterialInterface>> CombineMaterials;
-	MergedAttachmentMesh->EditMesh([&AttachmentItems, &CombineMaterials, this
 #if WITH_EDITOR
-		, &MaterialsMissingNiagaraUsageFlag
+	int32 ReadOnlyMaterialCount = 0;
 #endif
-		](FDynamicMesh3& InMergedMesh)
+
+	for (FCEClonerAttachmentItem* AttachmentItem : AttachmentItems)
 	{
-		InMergedMesh.Clear();
-		InMergedMesh.EnableAttributes();
-		InMergedMesh.Attributes()->SetNumNormalLayers(1);
+		if (!AttachmentItem)
+		{
+			continue;
+		}
 
-		FDynamicMeshEditor Editor(&InMergedMesh);
-		FGeometryScriptAppendMeshOptions AppendOptions;
-		AppendOptions.CombineMode = EGeometryScriptCombineAttributesMode::EnableAllMatching;
-		int32 MaterialCount = 0;
+		const UDynamicMesh* BakedDynamicMesh = AttachmentItem->BakedMesh;
+		if (!BakedDynamicMesh)
+		{
+			continue;
+		}
 
-		// Convert meshes
-		for (int32 MeshIdx = 0; MeshIdx < AttachmentItems.Num(); MeshIdx++)
+		FTransform MeshTransform = FTransform::Identity;
+		if (AttachmentItem->ParentActor.IsValid())
 		{
-			FCEClonerAttachmentItem* AttachmentItem = AttachmentItems[MeshIdx];
-			if (!AttachmentItem)
-			{
-				continue;
-			}
+			const FTransform& ParentTransform = AttachmentItem->ParentActor->GetTransform();
+			MeshTransform = AttachmentItem->ItemActor->GetTransform().GetRelativeTransform(ParentTransform);
+		}
 
-			const UDynamicMesh* BakedDynamicMesh = AttachmentItem->BakedMesh;
-			if (!BakedDynamicMesh)
-			{
-				continue;
-			}
+		TArray<TWeakObjectPtr<UMaterialInterface>> MeshMaterials;
+		for (TWeakObjectPtr<UMaterialInterface>& BakedMaterial : AttachmentItem->BakedMaterials)
+		{
+			UMaterialInterface* MaterialInterface = BakedMaterial.Get();
 
-			BakedDynamicMesh->ProcessMesh([&AttachmentItem, &CombineMaterials, &MaterialCount, &AppendOptions, &InMergedMesh, &Editor, this
 #if WITH_EDITOR
-				, &MaterialsMissingNiagaraUsageFlag
-#endif
-				](const FDynamicMesh3& InProcessMesh)
+			if (MaterialInterface && !IsMaterialUsageFlagSet(MaterialInterface))
 			{
-				FDynamicMesh3 ConvertedMesh = InProcessMesh;
-
-				if (AttachmentItem->ParentActor.IsValid())
+				if (!IsMaterialDirtyable(MaterialInterface))
 				{
-					const FTransform& ParentTransform = AttachmentItem->ParentActor->GetTransform();
-					const FTransform ParentChildTransform = AttachmentItem->ItemActor->GetTransform().GetRelativeTransform(ParentTransform);
-					MeshTransforms::ApplyTransform(ConvertedMesh, ParentChildTransform);
+					ReadOnlyMaterialCount++;
+					MaterialInterface = nullptr;
+					UE_LOG(LogCEClonerComponent, Warning, TEXT("%s : The following materials (%s) on actor (%s) does not have the required usage flag (bUsedWithNiagaraMeshParticles) to work with the cloner, this material cannot be dirtied due to its read-only location, skipping material and proceeding"), *ClonerActor->GetActorNameOrLabel(), MaterialInterface ? *MaterialInterface->GetMaterial()->GetPathName() : TEXT("Invalid Material"), *InRootActor->GetActorNameOrLabel());
 				}
 				else
 				{
-					MeshTransforms::ApplyTransform(ConvertedMesh, FTransform::Identity);
-				}
-
-				// Enable matching attributes & append mesh
-				FMeshIndexMappings TmpMappings;
-				AppendOptions.UpdateAttributesForCombineMode(InMergedMesh, ConvertedMesh);
-				Editor.AppendMesh(&ConvertedMesh, TmpMappings);
-
-				// Fix triangles materials linking
-				if (ConvertedMesh.HasAttributes() && ConvertedMesh.Attributes()->HasMaterialID())
-				{
-					const FDynamicMeshMaterialAttribute* FromMaterialIDAttrib = ConvertedMesh.Attributes()->GetMaterialID();
-					FDynamicMeshMaterialAttribute* ToMaterialIDAttrib = InMergedMesh.Attributes()->GetMaterialID();
-					TMap<int32, int32> MaterialMap;
-					for (const TPair<int32, int32>& FromToTId : TmpMappings.GetTriangleMap().GetForwardMap())
-					{
-						const int32 FromMatId = FromMaterialIDAttrib->GetValue(FromToTId.Key);
-						const int32 ToMatId = FromMatId + MaterialCount;
-						MaterialMap.Add(FromMatId, ToMatId);
-						ToMaterialIDAttrib->SetNewValue(FromToTId.Value, ToMatId);
-					}
-					MaterialCount += MaterialMap.Num();
-
-					for (TWeakObjectPtr<UMaterialInterface>& BakedMaterial : AttachmentItem->BakedMaterials)
-					{
-						UMaterialInterface* MaterialInterface = BakedMaterial.Get();
-						int32 MatIdx = CombineMaterials.Add(MaterialInterface);
-#if WITH_EDITOR
-						if (MaterialInterface && !IsMaterialUsageFlagSet(MaterialInterface))
-						{
-							MaterialsMissingNiagaraUsageFlag.Add(MatIdx);
-						}
-#endif
-					}
+					UE_LOG(LogCEClonerComponent, Log, TEXT("%s : The following materials (%s) on actor (%s) does not have the required usage flag (bUsedWithNiagaraMeshParticles) to work with the cloner, dirtying material and proceeding, please save the asset for working runtime result"), *ClonerActor->GetActorNameOrLabel(), MaterialInterface ? *MaterialInterface->GetMaterial()->GetPathName() : TEXT("Invalid Material"), *InRootActor->GetActorNameOrLabel());
 				}
-			});
-		}
-	}, EDynamicMeshChangeType::GeneralEdit, EDynamicMeshAttributeChangeFlags::Unknown, true);
-
-#if WITH_EDITOR
-	if (!MaterialsMissingNiagaraUsageFlag.IsEmpty())
-	{
-		int32 ReadOnlyMaterialCount = 0;
-
-		for (const int32 MaterialIdx : MaterialsMissingNiagaraUsageFlag)
-		{
-			UMaterialInterface* Material = CombineMaterials[MaterialIdx];
-			if (!Material)
-			{
-				continue;
 			}
+#endif
 
-			if (!IsMaterialDirtyable(Material))
-			{
-				ReadOnlyMaterialCount++;
-				CombineMaterials[MaterialIdx] = nullptr;
-				UE_LOG(LogCEClonerComponent, Warning, TEXT("%s : The following materials (%s) on actor (%s) does not have the required usage flag (bUsedWithNiagaraMeshParticles) to work with the cloner, this material cannot be dirtied due to its read-only location, skipping material and proceeding"), *ClonerActor->GetActorNameOrLabel(), Material ? *Material->GetMaterial()->GetPathName() : TEXT("Invalid Material"), *InRootActor->GetActorNameOrLabel());
-			}
-			else
-			{
-				UE_LOG(LogCEClonerComponent, Log, TEXT("%s : The following materials (%s) on actor (%s) does not have the required usage flag (bUsedWithNiagaraMeshParticles) to work with the cloner, dirtying material and proceeding, please save the asset for working runtime result"), *ClonerActor->GetActorNameOrLabel(), Material ? *Material->GetMaterial()->GetPathName() : TEXT("Invalid Material"), *InRootActor->GetActorNameOrLabel());
-			}
+			MeshMaterials.Add(MaterialInterface);
 		}
 
-		ShowMaterialWarning(ReadOnlyMaterialCount);
+		MeshBuilder.AppendMesh(AttachmentItem->BakedMesh, MeshMaterials, MeshTransform);
 	}
-#endif
 
-	// Convert merged mesh to static mesh
-	UStaticMesh* BakedAttachmentMesh = NewObject<UStaticMesh>();
-	if (MergedAttachmentMesh->GetTriangleCount())
-	{
-		// export options
-		FGeometryScriptCopyMeshToAssetOptions AssetOptions;
-		AssetOptions.bReplaceMaterials = true;
-		AssetOptions.bEnableRecomputeNormals = false;
-		AssetOptions.bEnableRecomputeTangents = false;
-		AssetOptions.bEnableRemoveDegenerates = true;
-		AssetOptions.bEmitTransaction = false;
-		AssetOptions.bApplyNaniteSettings = false;
-		AssetOptions.bDeferMeshPostEditChange = false;
-		AssetOptions.NewMaterials = CombineMaterials;
-		// LOD options
-		static FGeometryScriptMeshWriteLOD TargetLOD;
-		TargetLOD.LODIndex = 0;
-		TargetLOD.bWriteHiResSource = false;
-		// result
-		EGeometryScriptOutcomePins OutResult;
-		UGeometryScriptLibrary_StaticMeshFunctions::CopyMeshToStaticMesh(MergedAttachmentMesh, BakedAttachmentMesh, AssetOptions, TargetLOD, OutResult);
+	UStaticMesh* Mesh = NewObject<UStaticMesh>();
+	TArray<TWeakObjectPtr<UMaterialInterface>> MeshMaterials;
+	bClonerMeshesDirty = MeshBuilder.BuildStaticMesh(Mesh, MeshMaterials);
 
-#if WITH_EDITORONLY_DATA
-		// Compute normals and tangents
-		FMeshDescription* MeshDescription = BakedAttachmentMesh->GetMeshDescription(0);
-		FStaticMeshOperations::ComputeTriangleTangentsAndNormals(*MeshDescription);
-		FStaticMeshOperations::ComputeTangentsAndNormals(*MeshDescription, EComputeNTBsFlags::Normals);
-#endif
+	ClonerTree.MergedBakedMeshes[RootIdx] = Mesh;
 
-		bClonerMeshesDirty = true;
+#if WITH_EDITOR
+	if (ReadOnlyMaterialCount > 0)
+	{
+		ShowMaterialWarning(ReadOnlyMaterialCount);
 	}
-
-	MergedAttachmentMesh->MarkAsGarbage();
-	ClonerTree.MergedBakedMeshes[RootIdx] = BakedAttachmentMesh;
+#endif
 }
 
 void UCEClonerComponent::GetActorAttachmentItems(AActor* InActor, TArray<FCEClonerAttachmentItem*>& OutAttachmentItems)
@@ -1288,7 +943,7 @@ bool UCEClonerComponent::IsAllMergedMeshesValid() const
 {
 	for (const TObjectPtr<UStaticMesh>& MergedMesh : ClonerTree.MergedBakedMeshes)
 	{
-		if (!MergedMesh)
+		if (!MergedMesh.Get())
 		{
 			return false;
 		}
@@ -1340,136 +995,596 @@ void UCEClonerComponent::ShowMaterialWarning(int32 InMaterialCount)
 }
 #endif // WITH_EDITOR
 
-TArray<FNiagaraMeshMaterialOverride> UCEClonerComponent::GetOverrideMeshesMaterials() const
+void UCEClonerComponent::TickComponent(float InDeltaTime, ELevelTick InTickType, FActorComponentTickFunction* InThisTickFunction)
 {
-	TArray<FNiagaraMeshMaterialOverride> MaterialOverrides;
+	TickCloner(InDeltaTime);
 
-	if (bUseClonerMeshesOverrideMaterial)
-	{
-		int32 MaterialCount = 0;
-
-		for (const TWeakObjectPtr<AActor>& RootActor : ClonerTree.RootActors)
-		{
-			if (!RootActor.IsValid())
-			{
-				continue;
-			}
+	Super::TickComponent(InDeltaTime, InTickType, InThisTickFunction);
+}
 
-			const FCEClonerAttachmentItem* AttachmentItem = ClonerTree.ItemAttachmentMap.Find(RootActor);
+void UCEClonerComponent::InitializeCloner()
+{
+	if (bClonerInitialized)
+	{
+		return;
+	}
 
-			if (!AttachmentItem)
-			{
-				continue;
-			}
+	bClonerInitialized = true;
 
-			MaterialCount += AttachmentItem->BakedMaterials.Num();
-		}
+	/** Allow tick in editor without hitting check */
+	SetForceSolo(true);
 
-		// Set same material for all available slots
-		MaterialOverrides.Reserve(MaterialCount);
-		UMaterialInterface* OverrideMaterial = ClonerMeshesOverrideMaterial.Get();
+	OnLayoutNameChanged();
 
-		for (int32 Index = 0; Index < MaterialCount; Index++)
-		{
-			FNiagaraMeshMaterialOverride MaterialOverride;
-			MaterialOverride.ExplicitMat = OverrideMaterial;
-			MaterialOverrides.Add(MaterialOverride);
-		}
-	}
+#if WITH_EDITOR
+	OnVisualizerSpriteVisibleChanged();
+#endif
 
-	return MaterialOverrides;
+	OnClonerInitializedDelegate.Broadcast(this);
 }
 
-void UCEClonerComponent::SetUseOverrideMeshesMaterial(bool bInOverride)
+void UCEClonerComponent::TickCloner(float InDelta)
 {
-	if (bUseClonerMeshesOverrideMaterial == bInOverride)
+	if (!bClonerInitialized)
 	{
-		return;
+		InitializeCloner();
 	}
 
-	bUseClonerMeshesOverrideMaterial = bInOverride;
-
-	if (ClonerMeshesOverrideMaterial.IsValid())
+	if (bEnabled)
 	{
-		OnOverrideMeshesMaterialChanged();
-	}
-}
+		TreeUpdateDeltaTime += InDelta;
 
-void UCEClonerComponent::SetOverrideMeshesMaterial(UMaterialInterface* InMaterial)
-{
-	if (InMaterial && !IsMaterialUsageFlagSet(InMaterial))
-	{
-		UE_LOG(LogCEClonerComponent, Warning, TEXT("%s : The override material (%s) you wish to use does not have the required usage flag (bUsedWithNiagaraMeshParticles) to work with the cloner, enable the flag on the material and save the asset"), *GetOwner()->GetActorNameOrLabel(), *InMaterial->GetMaterial()->GetPathName());
+		if (TreeUpdateDeltaTime >= TreeUpdateInterval)
+		{
+			TreeUpdateDeltaTime -= TreeUpdateInterval;
 
-#if WITH_EDITOR
-		ShowMaterialWarning(1);
-#endif
+			UpdateClonerAttachmentTree();
+			UpdateClonerRenderState();
+		}
 
-		InMaterial = nullptr;
+		if (bNeedsRefresh)
+		{
+			bNeedsRefresh = false;
+			RequestClonerUpdate(true);
+		}
 	}
+}
 
-	if (ClonerMeshesOverrideMaterial == InMaterial)
+void UCEClonerComponent::SetEnabled(bool bInEnable)
+{
+	if (bInEnable == bEnabled)
 	{
 		return;
 	}
 
-	ClonerMeshesOverrideMaterial = InMaterial;
+	bEnabled = bInEnable;
+	OnEnabledChanged();
+}
 
-	if (bUseClonerMeshesOverrideMaterial)
+void UCEClonerComponent::SetTreeUpdateInterval(float InInterval)
+{
+	if (InInterval == TreeUpdateInterval)
 	{
-		OnOverrideMeshesMaterialChanged();
+		return;
 	}
+
+	TreeUpdateInterval = InInterval;
 }
 
-void UCEClonerComponent::OnOverrideMeshesMaterialChanged()
+void UCEClonerComponent::SetSeed(int32 InSeed)
 {
-	if (!bClonerMeshesUpdating && !bClonerMeshesDirty)
+	if (InSeed == Seed)
 	{
-		UpdateClonerMeshes();
+		return;
 	}
+
+	Seed = InSeed;
+	OnSeedChanged();
 }
 
-void UCEClonerComponent::UpdateClonerMeshes()
+void UCEClonerComponent::SetLayoutName(FName InLayoutName)
 {
-	const AActor* ClonerActor = GetOwner();
-	if (!ClonerActor)
+	if (LayoutName == InLayoutName)
 	{
 		return;
 	}
 
-	if (!GetAsset() || !ActiveLayout)
+	const TArray<FString> LayoutNames = GetClonerLayoutNames();
+	if (!LayoutNames.Contains(InLayoutName))
 	{
 		return;
 	}
 
-	UNiagaraMeshRendererProperties* MeshRenderer = ActiveLayout->GetMeshRenderer();
+	LayoutName = InLayoutName;
+	OnLayoutNameChanged();
+}
 
-	if (!MeshRenderer)
+#if WITH_EDITOR
+void UCEClonerComponent::SetVisualizerSpriteVisible(bool bInVisible)
+{
+	if (bVisualizerSpriteVisible == bInVisible)
 	{
-		UE_LOG(LogCEClonerComponent, Warning, TEXT("%s : Invalid mesh renderer for cloner system"), *ClonerActor->GetActorNameOrLabel());
 		return;
 	}
 
-	// Set material override
-	MeshRenderer->OverrideMaterials = GetOverrideMeshesMaterials();
-	MeshRenderer->bOverrideMaterials = bUseClonerMeshesOverrideMaterial;
+	bVisualizerSpriteVisible = bInVisible;
+	OnVisualizerSpriteVisibleChanged();
+}
+#endif
 
-	if (bClonerMeshesDirty)
+int32 UCEClonerComponent::GetMeshCount() const
+{
+	if (const UCEClonerLayoutBase* LayoutSystem = GetActiveLayout())
 	{
-		// Set new number of meshes in renderer
-		static const FName MeshNumName(TEXT("MeshNum"));
-		SetIntParameter(MeshNumName, ClonerTree.MergedBakedMeshes.Num());
-
-		// Resize mesh array properly
-		if (MeshRenderer->Meshes.Num() > ClonerTree.MergedBakedMeshes.Num())
+		if (const UNiagaraMeshRendererProperties* MeshRenderer = LayoutSystem->GetMeshRenderer())
 		{
-			MeshRenderer->Meshes.SetNum(ClonerTree.MergedBakedMeshes.Num());
+			return MeshRenderer->Meshes.Num();
+		}
+	}
+
+	return 0;
+}
+
+int32 UCEClonerComponent::GetAttachmentCount() const
+{
+	return ClonerTree.ItemAttachmentMap.Num();
+}
+
+#if WITH_EDITOR
+void UCEClonerComponent::ForceUpdateCloner()
+{
+	UpdateClonerAttachmentTree();
+	UpdateClonerRenderState();
+	OnLayoutNameChanged();
+}
+
+void UCEClonerComponent::CreateDefaultActorAttached()
+{
+	const UCEClonerEffectorSettings* ClonerEffectorSettings = GetDefault<UCEClonerEffectorSettings>();
+	if (!ClonerEffectorSettings || !ClonerEffectorSettings->GetSpawnDefaultActorAttached())
+	{
+		return;
+	}
+
+	// Only spawn if world is valid and not a preview actor
+	UWorld* World = GetWorld();
+	AActor* Owner = GetOwner();
+	if (!World || !Owner || Owner->bIsEditorPreviewActor)
+	{
+		return;
+	}
+
+	// Only spawn if no actor is attached below it
+	TArray<AActor*> AttachedActors;
+	constexpr bool bReset = true;
+	constexpr bool bRecursive = false;
+	Owner->GetAttachedActors(AttachedActors, bReset, bRecursive);
+
+	if (!AttachedActors.IsEmpty())
+	{
+		return;
+	}
+
+	UStaticMesh* DefaultStaticMesh = ClonerEffectorSettings->GetDefaultStaticMesh();
+	UMaterialInterface* DefaultMaterial = ClonerEffectorSettings->GetDefaultMaterial();
+
+	if (!DefaultStaticMesh || !DefaultMaterial)
+	{
+		return;
+	}
+
+	// Spawn attached actor with same flags as this actor
+	FActorSpawnParameters SpawnParameters;
+	SpawnParameters.Owner = Owner;
+	SpawnParameters.ObjectFlags = GetFlags();
+	SpawnParameters.bTemporaryEditorActor = false;
+
+	const FVector ClonerLocation = GetComponentLocation();
+	const FRotator ClonerRotation = GetComponentRotation();
+
+	if (AStaticMeshActor* DefaultActorAttached = World->SpawnActor<AStaticMeshActor>(ClonerLocation, ClonerRotation, SpawnParameters))
+	{
+		UStaticMeshComponent* StaticMeshComponent = DefaultActorAttached->GetStaticMeshComponent();
+		StaticMeshComponent->SetStaticMesh(DefaultStaticMesh);
+		StaticMeshComponent->SetMaterial(0, DefaultMaterial);
+
+		DefaultActorAttached->SetMobility(EComponentMobility::Movable);
+		DefaultActorAttached->AttachToActor(GetOwner(), FAttachmentTransformRules::KeepWorldTransform);
+
+		FActorLabelUtilities::SetActorLabelUnique(DefaultActorAttached, TEXT("DefaultClone"));
+	}
+}
+#endif
+
+void UCEClonerComponent::RequestClonerUpdate(bool bInImmediate)
+{
+	if (!bEnabled)
+	{
+		return;
+	}
+
+	if (bInImmediate)
+	{
+		bNeedsRefresh = false;
+		RefreshUserParameters();
+	}
+	else
+	{
+		bNeedsRefresh = true;
+	}
+}
+
+void UCEClonerComponent::OnEnabledChanged()
+{
+	if (bEnabled)
+	{
+		OnClonerEnabled();
+	}
+	else
+	{
+		OnClonerDisabled();
+	}
+}
+
+void UCEClonerComponent::OnClonerEnabled()
+{
+	for (const TObjectPtr<UCEClonerExtensionBase>& ActiveExtension : ActiveExtensions)
+	{
+		ActiveExtension->ActivateExtension();
+	}
+
+	OnLayoutNameChanged();
+}
+
+void UCEClonerComponent::OnClonerDisabled()
+{
+	for (const TObjectPtr<UCEClonerExtensionBase>& ActiveExtension : ActiveExtensions)
+	{
+		ActiveExtension->DeactivateExtension();
+	}
+
+	DeactivateImmediate();
+	SetAsset(nullptr);
+}
+
+void UCEClonerComponent::OnSeedChanged()
+{
+	if (!bEnabled)
+	{
+		return;
+	}
+
+	SetRandomSeedOffset(Seed);
+
+	RequestClonerUpdate();
+}
+
+void UCEClonerComponent::OnLayoutNameChanged()
+{
+	if (!bEnabled)
+	{
+		return;
+	}
+
+	UCEClonerLayoutBase* NewActiveLayout = FindOrAddLayout(LayoutName);
+
+	// Apply layout
+	SetClonerActiveLayout(NewActiveLayout);
+}
+
+#if WITH_EDITOR
+void UCEClonerComponent::OnVisualizerSpriteVisibleChanged()
+{
+	if (UTexture2D* SpriteTexture = LoadObject<UTexture2D>(nullptr, SpriteTexturePath))
+	{
+		CreateSpriteComponent(SpriteTexture);
+
+		if (SpriteComponent)
+		{
+			if (SpriteComponent->Sprite != SpriteTexture)
+			{
+				SpriteComponent->SetSprite(SpriteTexture);
+			}
+
+			SpriteComponent->SetVisibility(bVisualizerSpriteVisible, false);
+		}
+	}
+}
+#endif
+
+UCEClonerLayoutBase* UCEClonerComponent::FindOrAddLayout(TSubclassOf<UCEClonerLayoutBase> InClass)
+{
+	const UCEClonerSubsystem* Subsystem = UCEClonerSubsystem::Get();
+
+	if (!Subsystem)
+	{
+		return nullptr;
+	}
+
+	const FName ClassLayoutName = Subsystem->FindLayoutName(InClass);
+
+	if (ClassLayoutName.IsNone())
+	{
+		return nullptr;
+	}
+
+	return FindOrAddLayout(ClassLayoutName);
+}
+
+UCEClonerLayoutBase* UCEClonerComponent::FindOrAddLayout(FName InLayoutName)
+{
+	UCEClonerSubsystem* Subsystem = UCEClonerSubsystem::Get();
+	if (!Subsystem)
+	{
+		return nullptr;
+	}
+
+	// Check cached layout instances
+	UCEClonerLayoutBase* NewActiveLayout = nullptr;
+	for (TObjectPtr<UCEClonerLayoutBase>& LayoutInstance : LayoutInstances)
+	{
+		if (LayoutInstance && LayoutInstance->GetLayoutName() == InLayoutName)
+		{
+			NewActiveLayout = LayoutInstance;
+			break;
+		}
+	}
+
+	// Create new layout instance and cache it
+	if (!NewActiveLayout)
+	{
+		NewActiveLayout = Subsystem->CreateNewLayout(InLayoutName, this);
+		LayoutInstances.Add(NewActiveLayout);
+	}
+
+	return NewActiveLayout;
+}
+
+UCEClonerExtensionBase* UCEClonerComponent::FindOrAddExtension(TSubclassOf<UCEClonerExtensionBase> InClass)
+{
+	const UCEClonerSubsystem* Subsystem = UCEClonerSubsystem::Get();
+
+	if (!Subsystem)
+	{
+		return nullptr;
+	}
+
+	const FName ExtensionName = Subsystem->FindExtensionName(InClass);
+
+	if (ExtensionName.IsNone())
+	{
+		return nullptr;
+	}
+
+	return FindOrAddExtension(ExtensionName);
+}
+
+UCEClonerExtensionBase* UCEClonerComponent::FindOrAddExtension(FName InExtensionName)
+{
+	// Check cached extension instances
+	UCEClonerExtensionBase* NewActiveExtension = nullptr;
+	for (TObjectPtr<UCEClonerExtensionBase>& ExtensionInstance : ExtensionInstances)
+	{
+		if (ExtensionInstance && ExtensionInstance->GetExtensionName() == InExtensionName)
+		{
+			NewActiveExtension = ExtensionInstance;
+			break;
+		}
+	}
+
+	// Create new extension instance and cache it
+	if (!NewActiveExtension)
+	{
+		UCEClonerSubsystem* Subsystem = UCEClonerSubsystem::Get();
+		if (!Subsystem)
+		{
+			return nullptr;
+		}
+
+		NewActiveExtension = Subsystem->CreateNewExtension(InExtensionName, this);
+		ExtensionInstances.Add(NewActiveExtension);
+	}
+
+	return NewActiveExtension;
+}
+
+TArray<FString> UCEClonerComponent::GetClonerLayoutNames() const
+{
+	TArray<FString> LayoutNamesStrings;
+
+	if (const UCEClonerSubsystem* Subsystem = UCEClonerSubsystem::Get())
+	{
+		const TArray<FName>& LayoutNames = Subsystem->GetLayoutNames();
+		LayoutNamesStrings.Reserve(LayoutNames.Num());
+
+		Algo::Transform(LayoutNames, LayoutNamesStrings, [](const FName& InName)
+		{
+			return InName.ToString();
+		});
+	}
+
+	return LayoutNamesStrings;
+}
+
+void UCEClonerComponent::RefreshClonerMeshes()
+{
+	if (!bClonerMeshesUpdating && !bClonerMeshesDirty)
+	{
+		UpdateClonerMeshes();
+	}
+}
+
+UCEClonerExtensionBase* UCEClonerComponent::GetExtension(TSubclassOf<UCEClonerExtensionBase> InExtensionClass) const
+{
+	const UCEClonerSubsystem* Subsystem = UCEClonerSubsystem::Get();
+
+	if (!Subsystem)
+	{
+		return nullptr;
+	}
+
+	const FName ExtensionName = Subsystem->FindExtensionName(InExtensionClass.Get());
+
+	if (ExtensionName.IsNone())
+	{
+		return nullptr;
+	}
+
+	return GetExtension(ExtensionName);
+}
+
+UCEClonerExtensionBase* UCEClonerComponent::GetExtension(FName InExtensionName) const
+{
+	for (const TObjectPtr<UCEClonerExtensionBase>& ExtensionInstance : ExtensionInstances)
+	{
+		if (ExtensionInstance && ExtensionInstance->GetExtensionName() == InExtensionName)
+		{
+			return ExtensionInstance;
+		}
+	}
+
+	return nullptr;
+}
+
+void UCEClonerComponent::OnActiveLayoutLoaded(UCEClonerLayoutBase* InLayout, bool bInSuccess)
+{
+	if (!InLayout)
+	{
+		return;
+	}
+
+	InLayout->OnLayoutLoadedDelegate().RemoveAll(this);
+
+	if (!bInSuccess)
+	{
+		UE_LOG(LogCEClonerComponent, Warning, TEXT("%s : Cloner layout system failed to load %s - %s"), *GetOwner()->GetActorNameOrLabel(), *InLayout->GetLayoutName().ToString(), *InLayout->GetLayoutAssetPath());
+		return;
+	}
+
+	UE_LOG(LogCEClonerComponent, Log, TEXT("%s : Cloner layout system loaded %s - %s"), *GetOwner()->GetActorNameOrLabel(), *InLayout->GetLayoutName().ToString(), *InLayout->GetLayoutAssetPath());
+
+	OnClonerLayoutLoadedDelegate.Broadcast(this, InLayout);
+
+	ActivateLayout(InLayout);
+}
+
+void UCEClonerComponent::ActivateLayout(UCEClonerLayoutBase* InLayout)
+{
+	// Must be valid and loaded
+	if (!InLayout || !InLayout->IsLayoutLoaded())
+	{
+		return;
+	}
+
+	// Copy data interfaces to new layout
+	if (ActiveLayout && ActiveLayout->IsLayoutLoaded())
+	{
+		ActiveLayout->CopyTo(InLayout);
+	}
+
+	// Deactivate previous layout
+	if (ActiveLayout && ActiveLayout->IsLayoutActive())
+	{
+		ActiveLayout->DeactivateLayout();
+	}
+
+	// Activate new layout
+	InLayout->ActivateLayout();
+
+	ActiveLayout = InLayout;
+
+	UE_LOG(LogCEClonerComponent, Log, TEXT("%s : Cloner layout system changed %s - %s"), *GetOwner()->GetActorNameOrLabel(), *InLayout->GetLayoutName().ToString(), *InLayout->GetLayoutAssetPath());
+
+	OnActiveLayoutChanged();
+
+	bClonerMeshesDirty = true;
+}
+
+void UCEClonerComponent::OnActiveLayoutChanged()
+{
+	UCEClonerLayoutBase* Layout = GetActiveLayout();
+	if (!Layout)
+	{
+		return;
+	}
+
+	OnSeedChanged();
+
+	Layout->UpdateLayoutParameters();
+
+	TSet<TObjectPtr<UCEClonerExtensionBase>> PrevActiveExtensions(ActiveExtensions);
+	ActiveExtensions.Empty();
+
+	for (const FName& ExtensionName : Layout->GetSupportedExtensions())
+	{
+		if (UCEClonerExtensionBase* Extension = FindOrAddExtension(ExtensionName))
+		{
+			if (!PrevActiveExtensions.Contains(Extension))
+			{
+				Extension->ActivateExtension();
+			}
+
+			Extension->UpdateExtensionParameters();
+
+			ActiveExtensions.Add(Extension);
+			PrevActiveExtensions.Remove(Extension);
+		}
+	}
+
+	for (const TObjectPtr<UCEClonerExtensionBase>& InactiveExtension : PrevActiveExtensions)
+	{
+		InactiveExtension->DeactivateExtension();
+	}
+
+	ActiveExtensions.StableSort([](const TObjectPtr<UCEClonerExtensionBase>& InExtensionA, const TObjectPtr<UCEClonerExtensionBase>& InExtensionB)
+	{
+		return InExtensionA->GetExtensionPriority() > InExtensionB->GetExtensionPriority();
+	});
+}
+
+void UCEClonerComponent::UpdateClonerMeshes()
+{
+	const AActor* ClonerActor = GetOwner();
+
+	if (!ClonerActor)
+	{
+		return;
+	}
+
+	UNiagaraSystem* ActiveSystem = GetAsset();
+
+	if (!ActiveSystem || !ActiveLayout)
+	{
+		return;
+	}
+
+	if (ActiveLayout->GetSystem() != ActiveSystem)
+	{
+		UE_LOG(LogCEClonerComponent, Warning, TEXT("%s : Invalid asset for cloner layout"), *ClonerActor->GetActorNameOrLabel());
+		return;
+	}
+
+	UNiagaraMeshRendererProperties* MeshRenderer = ActiveLayout->GetMeshRenderer();
+
+	if (!MeshRenderer)
+	{
+		UE_LOG(LogCEClonerComponent, Warning, TEXT("%s : Invalid mesh renderer for cloner system"), *ClonerActor->GetActorNameOrLabel());
+		return;
+	}
+
+	if (bClonerMeshesDirty)
+	{
+		// Resize mesh array properly
+		if (MeshRenderer->Meshes.Num() > ClonerTree.MergedBakedMeshes.Num())
+		{
+			MeshRenderer->Meshes.SetNum(ClonerTree.MergedBakedMeshes.Num());
 		}
 
 		// Set baked meshes in mesh renderer array
 		for (int32 Idx = 0; Idx < ClonerTree.MergedBakedMeshes.Num(); Idx++)
 		{
-			UStaticMesh* StaticMesh = ClonerTree.MergedBakedMeshes[Idx];
+			UStaticMesh* StaticMesh = ClonerTree.MergedBakedMeshes[Idx].Get();
+
 			FNiagaraMeshRendererMeshProperties& MeshProperties = !MeshRenderer->Meshes.IsValidIndex(Idx) ? MeshRenderer->Meshes.AddDefaulted_GetRef() : MeshRenderer->Meshes[Idx];
 			MeshProperties.Mesh = StaticMesh && StaticMesh->GetNumTriangles(0) > 0 ? StaticMesh : nullptr;
 
@@ -1483,69 +1598,57 @@ void UCEClonerComponent::UpdateClonerMeshes()
 			}
 		}
 
-		bClonerMeshesDirty = !DirtyItemAttachments.IsEmpty();
+		bClonerMeshesDirty = !ClonerTree.DirtyItemAttachments.IsEmpty();
 
 		UE_LOG(LogCEClonerComponent, Log, TEXT("%s : Cloner mesh updated %i"), *ClonerActor->GetActorNameOrLabel(), ClonerTree.MergedBakedMeshes.Num())
 	}
 
+	for (const TObjectPtr<UCEClonerExtensionBase>& ActiveExtension : ActiveExtensions)
+	{
+		ActiveExtension->OnClonerMeshesUpdated();
+	}
+
+	// Set new number of meshes in renderer
+	SetIntParameter(TEXT("MeshNum"), MeshRenderer->Meshes.Num());
+
 #if WITH_EDITORONLY_DATA
 	MeshRenderer->OnMeshChanged();
 
 	// Used by other data interfaces to update their cached data
 	MeshRenderer->OnChanged().Broadcast();
+#else
+	FNiagaraSystemUpdateContext ReregisterContext(ActiveSystem, /** bReset */ true);
 #endif
 
-	OnClonerMeshUpdated.Broadcast(this);
+	OnClonerMeshUpdatedDelegate.Broadcast(this);
 }
 
-bool UCEClonerComponent::SetClonerActiveLayout(UCEClonerLayoutBase* InLayout)
+void UCEClonerComponent::SetClonerActiveLayout(UCEClonerLayoutBase* InLayout)
 {
 	if (!InLayout)
 	{
-		return false;
+		return;
 	}
 
 	const AActor* ClonerActor = GetOwner();
 	if (!ClonerActor)
 	{
-		return false;
+		return;
 	}
 
 	if (!InLayout->IsLayoutLoaded())
 	{
-		// Load system
-		if (!InLayout->LoadLayout())
+		if (!InLayout->OnLayoutLoadedDelegate().IsBoundToObject(this))
 		{
-			UE_LOG(LogCEClonerComponent, Warning, TEXT("%s : Cloner layout system failed to load %s - %s"), *GetOwner()->GetActorNameOrLabel(), *InLayout->GetLayoutName().ToString(), *InLayout->GetLayoutAssetPath());
-			return false;
+			InLayout->OnLayoutLoadedDelegate().AddUObject(this, &UCEClonerComponent::OnActiveLayoutLoaded);
 		}
 
-		UE_LOG(LogCEClonerComponent, Log, TEXT("%s : Cloner layout system loaded %s - %s"), *GetOwner()->GetActorNameOrLabel(), *InLayout->GetLayoutName().ToString(), *InLayout->GetLayoutAssetPath());
-
-		OnClonerSystemLoaded.Broadcast(this, InLayout);
-	}
-
-	// Copy data interfaces to new layout
-	if (ActiveLayout && ActiveLayout->IsLayoutLoaded())
-	{
-		ActiveLayout->CopyTo(InLayout);
-	}
+		InLayout->LoadLayout();
 
-	// Deactivate previous layout
-	if (ActiveLayout && ActiveLayout->IsLayoutActive())
-	{
-		ActiveLayout->DeactivateLayout();
+		return;
 	}
 
-	// Activate new layout
-	InLayout->ActivateLayout();
-
-	ActiveLayout = InLayout;
-	bClonerMeshesDirty = true;
-
-	UE_LOG(LogCEClonerComponent, Log, TEXT("%s : Cloner layout system changed %s - %s"), *ClonerActor->GetActorNameOrLabel(), *InLayout->GetLayoutName().ToString(), *InLayout->GetLayoutAssetPath());
-
-	return true;
+	ActivateLayout(InLayout);
 }
 
 #undef LOCTEXT_NAMESPACE
diff --git a/Engine/Plugins/Experimental/ClonerEffector/Source/ClonerEffector/Private/Cloner/CEClonerMeshBuilder.cpp b/Engine/Plugins/Experimental/ClonerEffector/Source/ClonerEffector/Private/Cloner/CEClonerMeshBuilder.cpp
new file mode 100644
--- /dev/null
+++ b/Engine/Plugins/Experimental/ClonerEffector/Source/ClonerEffector/Private/Cloner/CEClonerMeshBuilder.cpp
@@ -0,0 +1,555 @@
+// Copyright Epic Games, Inc. All Rights Reserved.
+
+#include "Cloner/CEClonerMeshBuilder.h"
+
+#include "Components/BrushComponent.h"
+#include "Components/DynamicMeshComponent.h"
+#include "Components/InstancedStaticMeshComponent.h"
+#include "Components/SkeletalMeshComponent.h"
+#include "Components/SplineMeshComponent.h"
+#include "DynamicMesh/MeshNormals.h"
+#include "DynamicMesh/MeshTransforms.h"
+#include "DynamicMesh/Operations/MergeCoincidentMeshEdges.h"
+#include "DynamicMeshEditor.h"
+#include "Engine/SkeletalMesh.h"
+#include "Engine/StaticMesh.h"
+#include "GeometryScript/MeshAssetFunctions.h"
+#include "GeometryScript/MeshBasicEditFunctions.h"
+#include "GeometryScript/SceneUtilityFunctions.h"
+#include "Materials/MaterialInterface.h"
+#include "ProceduralMeshComponent.h"
+#include "StaticMeshOperations.h"
+#include "UDynamicMesh.h"
+
+FCEClonerMeshBuilder::FCEClonerMeshBuilder()
+{
+	OutputDynamicMesh = NewObject<UDynamicMesh>();
+}
+
+void FCEClonerMeshBuilder::Reset()
+{
+	OutputDynamicMesh->EditMesh([](FDynamicMesh3& EditMesh){ EditMesh.Clear(); });
+	ConvertedMeshes.Empty();
+	OutputMeshMaterials.Empty();
+}
+
+int32 FCEClonerMeshBuilder::AppendActor(const AActor* InActor)
+{
+	int32 ComponentConverted = 0;
+
+	if (!IsValid(InActor))
+	{
+		return ComponentConverted;
+	}
+
+	const FTransform& SourceTransform = InActor->GetActorTransform();
+
+	TArray<UPrimitiveComponent*> PrimitiveComponents;
+	InActor->GetComponents(PrimitiveComponents, /** IncludeChildrenActors */ false);
+
+	for (UPrimitiveComponent* PrimitiveComponent : PrimitiveComponents)
+	{
+		if (!IsValid(PrimitiveComponent))
+		{
+			continue;
+		}
+
+		if (UDynamicMeshComponent* DynamicMeshComponent = Cast<UDynamicMeshComponent>(PrimitiveComponent))
+		{
+			if (AppendComponent(DynamicMeshComponent, SourceTransform))
+			{
+				ComponentConverted++;
+			}
+		}
+		else if (USkeletalMeshComponent* SkeletalMeshComponent = Cast<USkeletalMeshComponent>(PrimitiveComponent))
+		{
+			if (AppendComponent(SkeletalMeshComponent, SourceTransform))
+			{
+				ComponentConverted++;
+			}
+		}
+		else if (UBrushComponent* BrushComponent = Cast<UBrushComponent>(PrimitiveComponent))
+		{
+			if (AppendComponent(BrushComponent, SourceTransform))
+			{
+				ComponentConverted++;
+			}
+		}
+		else if (UProceduralMeshComponent* ProceduralMeshComponent = Cast<UProceduralMeshComponent>(PrimitiveComponent))
+		{
+			if (AppendComponent(ProceduralMeshComponent, SourceTransform))
+			{
+				ComponentConverted++;
+			}
+		}
+		else if (UInstancedStaticMeshComponent* ISMComponent = Cast<UInstancedStaticMeshComponent>(PrimitiveComponent))
+		{
+			if (AppendComponent(ISMComponent, SourceTransform))
+			{
+				ComponentConverted++;
+			}
+		}
+		else if (USplineMeshComponent* SplineMeshComponent = Cast<USplineMeshComponent>(PrimitiveComponent))
+		{
+			if (AppendComponent(SplineMeshComponent, SourceTransform))
+			{
+				ComponentConverted++;
+			}
+		}
+		else if (UStaticMeshComponent* StaticMeshComponent = Cast<UStaticMeshComponent>(PrimitiveComponent))
+		{
+			if (AppendComponent(StaticMeshComponent, SourceTransform))
+			{
+				ComponentConverted++;
+			}
+		}
+	}
+
+	return ComponentConverted;
+}
+
+bool FCEClonerMeshBuilder::AppendMesh(const UDynamicMesh* InMesh, const TArray<TWeakObjectPtr<UMaterialInterface>>& InMaterials, const FTransform& InTransform)
+{
+	if (!IsValid(InMesh))
+	{
+		return false;
+	}
+
+	// Create a copy of the mesh
+	InMesh->ProcessMesh([this, &InTransform](const FDynamicMesh3& EditMesh)
+	{
+		FDynamicMesh3 CopyMesh = EditMesh;
+		MeshTransforms::ApplyTransform(CopyMesh, InTransform);
+
+		ConvertedMeshes.Add(MoveTemp(CopyMesh));
+	});
+
+	OutputMeshMaterials.Append(InMaterials);
+
+	return true;
+}
+
+bool FCEClonerMeshBuilder::AppendComponent(const UStaticMeshComponent* InComponent, const FTransform& InSourceTransform)
+{
+	if (!IsValid(InComponent))
+	{
+		return false;
+	}
+
+	UStaticMesh* StaticMesh = InComponent->GetStaticMesh();
+
+	if (!IsValid(StaticMesh))
+	{
+		return false;
+	}
+
+	// convert to dynamic mesh
+	FGeometryScriptMeshReadLOD StaticMeshLOD;
+	StaticMeshLOD.LODType = EGeometryScriptLODType::RenderData;
+
+	FGeometryScriptCopyMeshFromAssetOptions OutputMeshOptions;
+	OutputMeshOptions.bIgnoreRemoveDegenerates = false;
+	OutputMeshOptions.bRequestTangents = false;
+	OutputMeshOptions.bApplyBuildSettings = false;
+
+	EGeometryScriptOutcomePins OutResult;
+	UGeometryScriptLibrary_StaticMeshFunctions::CopyMeshFromStaticMesh(StaticMesh, OutputDynamicMesh, OutputMeshOptions, StaticMeshLOD, OutResult);
+
+	if (OutResult != EGeometryScriptOutcomePins::Success)
+	{
+		return false;
+	}
+
+	// Transform the new mesh relative to the component
+	const FTransform RelativeTransform = InComponent->GetComponentTransform().GetRelativeTransform(InSourceTransform);
+	OutputMeshMaterials.Append(InComponent->GetMaterials());
+	OutputDynamicMesh->EditMesh([this, RelativeTransform](FDynamicMesh3& EditMesh)
+	{
+		MeshTransforms::ApplyTransform(EditMesh, RelativeTransform);
+		ConvertedMeshes.Add(MoveTemp(EditMesh));
+		// replace by empty mesh
+		FDynamicMesh3 EmptyMesh;
+		EditMesh = MoveTemp(EmptyMesh);
+	}, EDynamicMeshChangeType::GeneralEdit, EDynamicMeshAttributeChangeFlags::Unknown, true);
+
+	return true;
+}
+
+bool FCEClonerMeshBuilder::AppendComponent(UProceduralMeshComponent* InComponent, const FTransform& InSourceTransform)
+{
+	if (!IsValid(InComponent))
+	{
+		return false;
+	}
+
+	const int32 SectionCount = InComponent->GetNumSections();
+
+	if (SectionCount == 0)
+	{
+		return false;
+	}
+
+	// Transform the new mesh relative to the component
+	const FTransform RelativeTransform = InComponent->GetComponentTransform().GetRelativeTransform(InSourceTransform);
+
+	FDynamicMesh3 Mesh;
+	Mesh.EnableAttributes();
+	Mesh.Attributes()->EnablePrimaryColors();
+	Mesh.Attributes()->EnableMaterialID();
+	Mesh.Attributes()->SetNumNormalLayers(1);
+	Mesh.Attributes()->SetNumUVLayers(1);
+	Mesh.Attributes()->SetNumPolygroupLayers(1);
+	Mesh.Attributes()->EnableTangents();
+
+	using namespace UE::Geometry;
+
+	FDynamicMeshColorOverlay* ColorOverlay = Mesh.Attributes()->PrimaryColors();
+	FDynamicMeshNormalOverlay* NormalOverlay = Mesh.Attributes()->PrimaryNormals();
+	FDynamicMeshUVOverlay* UVOverlay = Mesh.Attributes()->PrimaryUV();
+	FDynamicMeshMaterialAttribute* MaterialAttr = Mesh.Attributes()->GetMaterialID();
+	FDynamicMeshPolygroupAttribute* GroupAttr = Mesh.Attributes()->GetPolygroupLayer(0);
+	FDynamicMeshNormalOverlay* TangentOverlay = Mesh.Attributes()->PrimaryTangents();
+
+	for (int32 SectionIdx = 0; SectionIdx < SectionCount; SectionIdx++)
+	{
+		if (FProcMeshSection* Section = InComponent->GetProcMeshSection(SectionIdx))
+		{
+			if (Section->bSectionVisible)
+			{
+				TArray<int32> VtxIds;
+				TArray<int32> NormalIds;
+				TArray<int32> ColorIds;
+				TArray<int32> UVIds;
+				TArray<int32> TaIds;
+
+				// copy vertices data (position, normal, color, UV, tangent)
+				for (FProcMeshVertex& SectionVertex : Section->ProcVertexBuffer)
+				{
+					int32 VId = Mesh.AppendVertex(SectionVertex.Position);
+					VtxIds.Add(VId);
+
+					int32 NId = NormalOverlay->AppendElement(static_cast<FVector3f>(SectionVertex.Normal));
+					NormalIds.Add(NId);
+
+					int32 CId = ColorOverlay->AppendElement(static_cast<FVector4f>(SectionVertex.Color));
+					ColorIds.Add(CId);
+
+					int32 UVId = UVOverlay->AppendElement(static_cast<FVector2f>(SectionVertex.UV0));
+					UVIds.Add(UVId);
+
+					int32 TaId = TangentOverlay->AppendElement(static_cast<FVector3f>(SectionVertex.Tangent.TangentX));
+					TaIds.Add(TaId);
+				}
+
+				// copy tris data
+				if (Section->ProcIndexBuffer.Num() % 3 != 0)
+				{
+					continue;
+				}
+				for (int32 Idx = 0; Idx < Section->ProcIndexBuffer.Num(); Idx+=3)
+				{
+					int32 VIdx1 = Section->ProcIndexBuffer[Idx];
+					int32 VIdx2 = Section->ProcIndexBuffer[Idx + 1];
+					int32 VIdx3 = Section->ProcIndexBuffer[Idx + 2];
+
+					int32 VId1 = VtxIds[VIdx1];
+					int32 VId2 = VtxIds[VIdx2];
+					int32 VId3 = VtxIds[VIdx3];
+
+					int32 TId = Mesh.AppendTriangle(VId1, VId2, VId3, SectionIdx);
+					if (TId < 0)
+					{
+						continue;
+					}
+
+					NormalOverlay->SetTriangle(TId, FIndex3i(NormalIds[VIdx1], NormalIds[VIdx2], NormalIds[VIdx3]), true);
+					ColorOverlay->SetTriangle(TId, FIndex3i(ColorIds[VIdx1], ColorIds[VIdx2], ColorIds[VIdx3]), true);
+					UVOverlay->SetTriangle(TId, FIndex3i(UVIds[VIdx1], UVIds[VIdx2], UVIds[VIdx3]), true);
+					TangentOverlay->SetTriangle(TId, FIndex3i(TaIds[VIdx1], TaIds[VIdx2], TaIds[VIdx3]), true);
+
+					MaterialAttr->SetValue(TId, SectionIdx);
+					GroupAttr->SetValue(TId, SectionIdx);
+				}
+			}
+		}
+	}
+
+	if (Mesh.TriangleCount() == 0)
+	{
+		return false;
+	}
+
+	MeshTransforms::ApplyTransform(Mesh, RelativeTransform);
+
+	OutputMeshMaterials.Append(InComponent->GetMaterials());
+	ConvertedMeshes.Add(MoveTemp(Mesh));
+
+	return true;
+}
+
+bool FCEClonerMeshBuilder::AppendComponent(UBrushComponent* InComponent, const FTransform& InSourceTransform)
+{
+	return AppendPrimitiveComponent(InComponent, InSourceTransform);
+}
+
+bool FCEClonerMeshBuilder::AppendComponent(const USkeletalMeshComponent* InComponent, const FTransform& InSourceTransform)
+{
+	if (!IsValid(InComponent))
+	{
+		return false;
+	}
+
+	USkeletalMesh* SkeletalMesh = InComponent->GetSkeletalMeshAsset();
+
+	if (!IsValid(SkeletalMesh))
+	{
+		return false;
+	}
+
+	// convert to dynamic mesh
+	FGeometryScriptMeshReadLOD SkeletalMeshLOD;
+	SkeletalMeshLOD.LODType = EGeometryScriptLODType::SourceModel;
+
+	FGeometryScriptCopyMeshFromAssetOptions OutputMeshOptions;
+	OutputMeshOptions.bIgnoreRemoveDegenerates = false;
+	OutputMeshOptions.bRequestTangents = false;
+	OutputMeshOptions.bApplyBuildSettings = false;
+
+	EGeometryScriptOutcomePins OutResult;
+	UGeometryScriptLibrary_StaticMeshFunctions::CopyMeshFromSkeletalMesh(SkeletalMesh, OutputDynamicMesh, OutputMeshOptions, SkeletalMeshLOD, OutResult);
+
+	if (OutResult != EGeometryScriptOutcomePins::Success)
+	{
+		return false;
+	}
+
+	// Transform the new mesh relative to the component
+	const FTransform RelativeTransform = InComponent->GetComponentTransform().GetRelativeTransform(InSourceTransform);
+
+	// Copy materials
+	OutputMeshMaterials.Append(InComponent->GetMaterials());
+
+	OutputDynamicMesh->EditMesh([this, RelativeTransform](FDynamicMesh3& EditMesh)
+	{
+		MeshTransforms::ApplyTransform(EditMesh, RelativeTransform);
+		ConvertedMeshes.Add(MoveTemp(EditMesh));
+
+		// replace by empty mesh for next usage
+		FDynamicMesh3 EmptyMesh;
+		EditMesh = MoveTemp(EmptyMesh);
+	}, EDynamicMeshChangeType::GeneralEdit, EDynamicMeshAttributeChangeFlags::Unknown, /** bDeferChanges */true);
+
+	return true;
+}
+
+bool FCEClonerMeshBuilder::AppendComponent(UDynamicMeshComponent* InComponent, const FTransform& InSourceTransform)
+{
+	if (!IsValid(InComponent))
+	{
+		return false;
+	}
+
+	const UDynamicMesh* DynamicMesh = InComponent->GetDynamicMesh();
+
+	if (!IsValid(DynamicMesh))
+	{
+		return false;
+	}
+
+	// Transform the new mesh relative to the component
+	const FTransform RelativeTransform = InComponent->GetComponentTransform().GetRelativeTransform(InSourceTransform);
+
+	// Copy all materials
+	OutputMeshMaterials.Append(InComponent->GetMaterials());
+
+	// Create a copy of the mesh
+	DynamicMesh->ProcessMesh([this, RelativeTransform](const FDynamicMesh3& EditMesh)
+	{
+		FDynamicMesh3 CopyMesh = EditMesh;
+		MeshTransforms::ApplyTransform(CopyMesh, RelativeTransform);
+		ConvertedMeshes.Add(MoveTemp(CopyMesh));
+	});
+
+	return true;
+}
+
+bool FCEClonerMeshBuilder::AppendComponent(UInstancedStaticMeshComponent* InComponent, const FTransform& InSourceTransform)
+{
+	if (!IsValid(InComponent) || !IsValid(InComponent->GetStaticMesh()))
+	{
+		return false;
+	}
+
+	return AppendPrimitiveComponent(InComponent, InSourceTransform);
+}
+
+bool FCEClonerMeshBuilder::AppendComponent(USplineMeshComponent* InComponent, const FTransform& InSourceTransform)
+{
+	if (!IsValid(InComponent) || !IsValid(InComponent->GetStaticMesh()))
+	{
+		return false;
+	}
+
+	return AppendPrimitiveComponent(InComponent, InSourceTransform);
+}
+
+bool FCEClonerMeshBuilder::BuildDynamicMesh(UDynamicMesh* OutMesh, TArray<TWeakObjectPtr<UMaterialInterface>>& OutMaterials)
+{
+	if (!IsValid(OutMesh))
+	{
+		return false;
+	}
+
+	// Lets combine all meshes components from this actor together
+	OutMesh->EditMesh([this](FDynamicMesh3& InMergedMesh)
+	{
+		using namespace UE::Geometry;
+
+		InMergedMesh.Clear();
+		InMergedMesh.EnableAttributes();
+		InMergedMesh.Attributes()->SetNumNormalLayers(1);
+
+		FDynamicMeshEditor Editor(&InMergedMesh);
+		FGeometryScriptAppendMeshOptions AppendOptions;
+		AppendOptions.CombineMode = EGeometryScriptCombineAttributesMode::EnableAllMatching;
+		int32 MaterialCount = 0;
+
+		// Convert meshes
+		for (int32 MeshIdx = 0; MeshIdx < ConvertedMeshes.Num(); MeshIdx++)
+		{
+			const FDynamicMesh3& ConvertedMesh = ConvertedMeshes[MeshIdx];
+
+			// Enable matching attributes & append mesh
+			FMeshIndexMappings TmpMappings;
+			AppendOptions.UpdateAttributesForCombineMode(InMergedMesh, ConvertedMesh);
+			Editor.AppendMesh(&ConvertedMesh, TmpMappings);
+
+			// Fix triangles materials linking
+			if (ConvertedMesh.HasAttributes() && ConvertedMesh.Attributes()->HasMaterialID())
+			{
+				const FDynamicMeshMaterialAttribute* FromMaterialIDAttrib = ConvertedMesh.Attributes()->GetMaterialID();
+				FDynamicMeshMaterialAttribute* ToMaterialIDAttrib = InMergedMesh.Attributes()->GetMaterialID();
+				TMap<int32, int32> MaterialMap;
+				for (const TPair<int32, int32>& FromToTId : TmpMappings.GetTriangleMap().GetForwardMap())
+				{
+					const int32 FromMatId = FromMaterialIDAttrib->GetValue(FromToTId.Key);
+					const int32 ToMatId = FromMatId + MaterialCount;
+					MaterialMap.Add(FromMatId, ToMatId);
+					ToMaterialIDAttrib->SetNewValue(FromToTId.Value, ToMatId);
+				}
+				MaterialCount += MaterialMap.Num();
+			}
+		}
+
+		if (InMergedMesh.TriangleCount() > 0)
+		{
+			// Merge shared edges
+			FMergeCoincidentMeshEdges WeldOp(&InMergedMesh);
+			WeldOp.Apply();
+		}
+	}, EDynamicMeshChangeType::GeneralEdit, EDynamicMeshAttributeChangeFlags::Unknown, /** bDeferChange */true);
+
+	OutMaterials = MoveTemp(OutputMeshMaterials);
+
+	ConvertedMeshes.Empty();
+	OutputMeshMaterials = {};
+
+	return true;
+}
+
+bool FCEClonerMeshBuilder::BuildStaticMesh(UStaticMesh* OutMesh, TArray<TWeakObjectPtr<UMaterialInterface>>& OutMaterials)
+{
+	if (!IsValid(OutMesh))
+	{
+		return false;
+	}
+
+	UDynamicMesh* OutDynamicMesh = OutputDynamicMesh;
+	if (!BuildDynamicMesh(OutDynamicMesh, OutMaterials))
+	{
+		return false;
+	}
+
+	if (OutDynamicMesh->GetTriangleCount() == 0)
+	{
+		return false;
+	}
+
+	TArray<TObjectPtr<UMaterialInterface>> NewMaterials;
+	Algo::Transform(OutMaterials, NewMaterials, [](const TWeakObjectPtr<UMaterialInterface>& InMaterialWeak)
+	{
+		return InMaterialWeak.Get();
+	});
+
+	// Export options
+	FGeometryScriptCopyMeshToAssetOptions AssetOptions;
+	AssetOptions.bReplaceMaterials = true;
+	AssetOptions.bEnableRecomputeNormals = false;
+	AssetOptions.bEnableRecomputeTangents = false;
+	AssetOptions.bEnableRemoveDegenerates = true;
+	AssetOptions.bEmitTransaction = false;
+	AssetOptions.bApplyNaniteSettings = false;
+	AssetOptions.bDeferMeshPostEditChange = false;
+	AssetOptions.NewMaterials = NewMaterials;
+
+	// LOD options
+	FGeometryScriptMeshWriteLOD TargetLOD;
+	TargetLOD.LODIndex = 0;
+	TargetLOD.bWriteHiResSource = false;
+
+	// Convert merged mesh to static mesh
+	EGeometryScriptOutcomePins OutResult;
+	UGeometryScriptLibrary_StaticMeshFunctions::CopyMeshToStaticMesh(OutDynamicMesh, OutMesh, AssetOptions, TargetLOD, OutResult);
+
+#if WITH_EDITORONLY_DATA
+	// Compute normals and tangents
+	FMeshDescription* MeshDescription = OutMesh->GetMeshDescription(TargetLOD.LODIndex);
+	FStaticMeshOperations::ComputeTriangleTangentsAndNormals(*MeshDescription);
+	FStaticMeshOperations::ComputeTangentsAndNormals(*MeshDescription, EComputeNTBsFlags::Normals);
+#endif
+
+	OutDynamicMesh->EditMesh([](FDynamicMesh3& EditMesh){ EditMesh.Clear(); });
+
+	return true;
+}
+
+bool FCEClonerMeshBuilder::AppendPrimitiveComponent(UPrimitiveComponent* InComponent, const FTransform& InSourceTransform)
+{
+	if (!IsValid(InComponent))
+	{
+		return false;
+	}
+
+	constexpr FGeometryScriptCopyMeshFromComponentOptions Options;
+	FTransform LocalToWorld = FTransform::Identity;
+
+	// convert to dynamic mesh
+	EGeometryScriptOutcomePins OutResult;
+	UGeometryScriptLibrary_SceneUtilityFunctions::CopyMeshFromComponent(InComponent, OutputDynamicMesh, Options, false, LocalToWorld, OutResult);
+
+	if (OutResult != EGeometryScriptOutcomePins::Success)
+	{
+		return false;
+	}
+
+	// Transform the new mesh relative to the component
+	const FTransform RelativeTransform = InComponent->GetComponentTransform().GetRelativeTransform(InSourceTransform);
+
+	// Copy all materials
+	for (int32 Index = 0; Index < InComponent->GetNumMaterials(); Index++)
+	{
+		OutputMeshMaterials.Add(InComponent->GetMaterial(Index));
+	}
+
+	OutputDynamicMesh->EditMesh([this, RelativeTransform](FDynamicMesh3& EditMesh)
+	{
+		MeshTransforms::ApplyTransform(EditMesh, RelativeTransform);
+		ConvertedMeshes.Add(MoveTemp(EditMesh));
+
+		// replace by empty mesh
+		FDynamicMesh3 EmptyMesh;
+		EditMesh = MoveTemp(EmptyMesh);
+	}, EDynamicMeshChangeType::GeneralEdit, EDynamicMeshAttributeChangeFlags::Unknown, /** bDeferChange */true);
+
+	return true;
+}
diff --git a/Engine/Plugins/Experimental/ClonerEffector/Source/ClonerEffector/Private/Cloner/Extensions/CEClonerEffectorExtension.cpp b/Engine/Plugins/Experimental/ClonerEffector/Source/ClonerEffector/Private/Cloner/Extensions/CEClonerEffectorExtension.cpp
new file mode 100644
--- /dev/null
+++ b/Engine/Plugins/Experimental/ClonerEffector/Source/ClonerEffector/Private/Cloner/Extensions/CEClonerEffectorExtension.cpp
@@ -0,0 +1,277 @@
+// Copyright Epic Games, Inc. All Rights Reserved.
+
+#include "Cloner/Extensions/CEClonerEffectorExtension.h"
+
+#include "Cloner/CEClonerComponent.h"
+#include "Cloner/Layouts/CEClonerLayoutBase.h"
+#include "Effector/CEEffectorActor.h"
+#include "Effector/CEEffectorComponent.h"
+#include "Subsystems/CEEffectorSubsystem.h"
+
+#if WITH_EDITOR
+#include "Editor.h"
+#include "Editor/EditorEngine.h"
+#endif
+
+DEFINE_LOG_CATEGORY_STATIC(LogCEClonerEffectorExtension, Log, All);
+
+void UCEClonerEffectorExtension::SetColor(const FLinearColor& InColor)
+{
+	if (InColor.Equals(Color))
+	{
+		return;
+	}
+
+	Color = InColor;
+	UpdateExtensionParameters(/** UpdateCloner */false);
+}
+
+bool UCEClonerEffectorExtension::LinkEffector(AActor* InEffectorActor)
+{
+	if (!IsValid(InEffectorActor)
+		|| EffectorActorsWeak.Contains(InEffectorActor)
+		|| !InEffectorActor->FindComponentByClass<UCEEffectorComponent>())
+	{
+		return false;
+	}
+
+	EffectorActorsWeak.Add(InEffectorActor);
+
+	OnEffectorsChanged();
+
+	UE_LOG(LogCEClonerEffectorExtension, Log, TEXT("%s : Effector %s linked to Cloner"), *GetClonerComponent()->GetOwner()->GetActorNameOrLabel(), *InEffectorActor->GetActorNameOrLabel());
+
+	return true;
+}
+
+bool UCEClonerEffectorExtension::UnlinkEffector(AActor* InEffectorActor)
+{
+	if (!InEffectorActor)
+	{
+		return false;
+	}
+
+	if (EffectorActorsWeak.Remove(InEffectorActor) > 0)
+	{
+		OnEffectorsChanged();
+
+		UE_LOG(LogCEClonerEffectorExtension, Log, TEXT("%s : Effector %s unlinked from Cloner"), *GetClonerComponent()->GetOwner()->GetActorNameOrLabel(), *InEffectorActor->GetActorNameOrLabel());
+	}
+
+	return true;
+}
+
+bool UCEClonerEffectorExtension::IsEffectorLinked(const AActor* InEffectorActor) const
+{
+	return InEffectorActor && EffectorActorsWeak.Contains(InEffectorActor);
+}
+
+int32 UCEClonerEffectorExtension::GetEffectorCount() const
+{
+	return EffectorActorsWeak.Num();
+}
+
+const FCEClonerEffectorDataInterfaces* UCEClonerEffectorExtension::GetEffectorDataInterfaces() const
+{
+	if (const UCEClonerLayoutBase* LayoutSystem = GetClonerLayout())
+	{
+		return &LayoutSystem->GetDataInterfaces();
+	}
+
+	return nullptr;
+}
+
+void UCEClonerEffectorExtension::OnExtensionActivated()
+{
+	Super::OnExtensionActivated();
+
+	UCEEffectorSubsystem::OnEffectorIdentifierChanged().RemoveAll(this);
+	UCEEffectorSubsystem::OnEffectorIdentifierChanged().AddUObject(this, &UCEClonerEffectorExtension::OnEffectorIdentifierChanged);
+
+	OnEffectorsChanged();
+}
+
+void UCEClonerEffectorExtension::OnExtensionDeactivated()
+{
+	Super::OnExtensionDeactivated();
+
+	UCEEffectorSubsystem::OnEffectorIdentifierChanged().RemoveAll(this);
+}
+
+void UCEClonerEffectorExtension::OnExtensionParametersChanged(UCEClonerComponent* InComponent)
+{
+	Super::OnExtensionParametersChanged(InComponent);
+
+	InComponent->SetColorParameter(TEXT("EffectorDefaultColor"), Color);
+}
+
+void UCEClonerEffectorExtension::OnEffectorIdentifierChanged(UCEEffectorComponent* InEffector, int32 InOldIdentifier, int32 InNewIdentifier)
+{
+	if (EffectorActorsWeak.Contains(InEffector->GetOwner()))
+	{
+		OnEffectorsChanged();
+	}
+}
+
+void UCEClonerEffectorExtension::OnEffectorsChanged()
+{
+	const FCEClonerEffectorDataInterfaces* EffectorDataInterfaces = GetEffectorDataInterfaces();
+
+	if (!EffectorDataInterfaces || !EffectorDataInterfaces->GetIndexArray())
+	{
+		return;
+	}
+
+	// Remove duplicates
+	TSet<TWeakObjectPtr<UCEEffectorComponent>> SetEffectorsWeak;
+
+	for (TArray<TWeakObjectPtr<AActor>>::TIterator It(EffectorActorsWeak); It; ++It)
+	{
+		if (AActor* EffectorActor = It->Get())
+		{
+			TArray<UCEEffectorComponent*> EffectorComponents;
+			EffectorActor->GetComponents(EffectorComponents, /** Recurse */false);
+
+			for (UCEEffectorComponent* EffectorComponent : EffectorComponents)
+			{
+				if (EffectorComponent->GetChannelIdentifier() != INDEX_NONE)
+				{
+					SetEffectorsWeak.Add(EffectorComponent);
+				}
+			}
+
+			if (SetEffectorsWeak.IsEmpty())
+			{
+#if WITH_EDITOR
+				if (GUndo)
+				{
+					Modify();
+				}
+#endif
+
+				It.RemoveCurrent();
+			}
+		}
+	}
+
+	// Removed effectors
+	for (const TWeakObjectPtr<UCEEffectorComponent>& EffectorWeak : EffectorsInternalWeak.Difference(SetEffectorsWeak))
+	{
+		if (UCEEffectorComponent* Effector = EffectorWeak.Get())
+		{
+			Effector->OnClonerUnlinked(this);
+		}
+	}
+
+	// Added effectors
+	for (const TWeakObjectPtr<UCEEffectorComponent>& EffectorWeak : SetEffectorsWeak.Difference(EffectorsInternalWeak))
+	{
+		if (UCEEffectorComponent* Effector = EffectorWeak.Get())
+		{
+			Effector->OnClonerLinked(this);
+		}
+	}
+
+	EffectorsInternalWeak = SetEffectorsWeak;
+
+	TArray<int32> EffectorIndexes;
+	EffectorIndexes.Reserve(SetEffectorsWeak.Num());
+
+	for (const TWeakObjectPtr<UCEEffectorComponent>& EffectorWeak : SetEffectorsWeak)
+	{
+		if (UCEEffectorComponent* Effector = EffectorWeak.Get())
+		{
+			const int32 ChannelIdentifier = Effector->GetChannelIdentifier();
+
+			if (ChannelIdentifier != INDEX_NONE)
+			{
+				EffectorIndexes.AddUnique(ChannelIdentifier);
+			}
+		}
+	}
+
+	TArray<int32>& EffectorIndexArray = EffectorDataInterfaces->GetIndexArray()->GetArrayReference();
+	EffectorIndexArray.Empty(EffectorIndexes.Num());
+
+	for (const int32 EffectorIndex : EffectorIndexes)
+	{
+		EffectorIndexArray.Add(EffectorIndex);
+	}
+
+	UpdateExtensionParameters(/** UpdateCloner */true);
+}
+
+void UCEClonerEffectorExtension::OnEffectorActorsChanged()
+{
+	for (TArray<TWeakObjectPtr<AActor>>::TIterator It(EffectorActorsWeak); It; ++It)
+	{
+		if (AActor* EffectorActor = It->Get())
+		{
+			if (!EffectorActor->FindComponentByClass<UCEEffectorComponent>())
+			{
+				It.RemoveCurrent();
+			}
+		}
+	}
+
+	OnEffectorsChanged();
+}
+
+#if WITH_EDITOR
+void UCEClonerEffectorExtension::CreateLinkedEffector()
+{
+	UCEClonerComponent* ClonerComponent = GetClonerComponent();
+	if (!IsValid(ClonerComponent))
+	{
+		return;
+	}
+
+	UWorld* ClonerWorld = ClonerComponent->GetWorld();
+	if (!IsValid(ClonerWorld))
+	{
+		return;
+	}
+
+	FActorSpawnParameters Params;
+	Params.bTemporaryEditorActor = false;
+
+	const FVector ClonerLocation = ClonerComponent->GetComponentLocation();
+	const FRotator ClonerRotation = ClonerComponent->GetComponentRotation();
+
+	ACEEffectorActor* EffectorActor = ClonerWorld->SpawnActor<ACEEffectorActor>(ACEEffectorActor::StaticClass(), ClonerLocation, ClonerRotation, Params);
+	if (!EffectorActor)
+	{
+		return;
+	}
+
+	FActorLabelUtilities::RenameExistingActor(EffectorActor, EffectorActor->GetDefaultActorLabel(), true);
+
+	LinkEffector(EffectorActor);
+
+	if (GEditor)
+	{
+		GEditor->SelectNone(true, true);
+		GEditor->SelectActor(EffectorActor, true, true);
+	}
+}
+
+const TCEPropertyChangeDispatcher<UCEClonerEffectorExtension> UCEClonerEffectorExtension::PropertyChangeDispatcher =
+{
+	{ GET_MEMBER_NAME_CHECKED(UCEClonerEffectorExtension, Color), &UCEClonerEffectorExtension::OnExtensionPropertyChanged },
+	{ GET_MEMBER_NAME_CHECKED(UCEClonerEffectorExtension, EffectorActorsWeak), &UCEClonerEffectorExtension::OnEffectorActorsChanged }
+};
+
+void UCEClonerEffectorExtension::PostEditChangeProperty(FPropertyChangedEvent& InPropertyChangedEvent)
+{
+	Super::PostEditChangeProperty(InPropertyChangedEvent);
+
+	PropertyChangeDispatcher.OnPropertyChanged(this, InPropertyChangedEvent);
+}
+
+void UCEClonerEffectorExtension::PostEditUndo()
+{
+	Super::PostEditUndo();
+
+	OnEffectorsChanged();
+}
+#endif
diff --git a/Engine/Plugins/Experimental/ClonerEffector/Source/ClonerEffector/Private/Cloner/Layouts/CEClonerLayoutBase.cpp b/Engine/Plugins/Experimental/ClonerEffector/Source/ClonerEffector/Private/Cloner/Layouts/CEClonerLayoutBase.cpp
--- a/Engine/Plugins/Experimental/ClonerEffector/Source/ClonerEffector/Private/Cloner/Layouts/CEClonerLayoutBase.cpp
+++ b/Engine/Plugins/Experimental/ClonerEffector/Source/ClonerEffector/Private/Cloner/Layouts/CEClonerLayoutBase.cpp
@@ -4,9 +4,11 @@
 
 #include "Cloner/CEClonerActor.h"
 #include "Cloner/CEClonerComponent.h"
+#include "Misc/PackageName.h"
 #include "NiagaraEmitter.h"
 #include "NiagaraMeshRendererProperties.h"
 #include "NiagaraSystem.h"
+#include "Subsystems/CEClonerSubsystem.h"
 #include "UObject/Package.h"
 #include "UObject/UObjectGlobals.h"
 
@@ -20,10 +22,10 @@ bool UCEClonerLayoutBase::IsLayoutValid() const
 	}
 
 	// Get the template niagara asset
-	const UNiagaraSystem* TemplateNiagaraSystem = LoadSystemPath(LayoutAssetPath);
+	const UNiagaraSystem* TemplateNiagaraSystem = LoadObject<UNiagaraSystem>(nullptr, *LayoutAssetPath);
 
 	// Get the base niagara asset
-	const UNiagaraSystem* BaseNiagaraSystem = LoadSystemPath(LayoutBaseAssetPath);
+	const UNiagaraSystem* BaseNiagaraSystem = LoadObject<UNiagaraSystem>(nullptr, LayoutBaseAssetPath);
 
 	if (!TemplateNiagaraSystem || !BaseNiagaraSystem)
 	{
@@ -65,64 +67,54 @@ bool UCEClonerLayoutBase::IsLayoutLoaded() const
 	return !IsTemplate() && NiagaraSystem && MeshRenderer && DataInterfaces.IsValid();
 }
 
-bool UCEClonerLayoutBase::LoadLayout()
+void UCEClonerLayoutBase::LoadLayout()
 {
 	if (IsLayoutLoaded())
 	{
-		return true;
+		return;
 	}
 
-	if (LayoutAssetPath.IsEmpty())
+	// Already being loaded
+	if (LoadRequestIdentifier != INDEX_NONE)
 	{
-		return false;
+		return;
 	}
 
-	// Get the template niagara asset
-	const UNiagaraSystem* TemplateNiagaraSystem = FindObject<UNiagaraSystem>(nullptr, *LayoutAssetPath);
-	if (!TemplateNiagaraSystem)
+	if (LayoutAssetPath.IsEmpty())
 	{
-		TemplateNiagaraSystem = LoadObject<UNiagaraSystem>(nullptr, *LayoutAssetPath);
+		return;
 	}
 
-	if (!TemplateNiagaraSystem)
+	UCEClonerComponent* ClonerComponent = GetClonerComponent();
+
+	if (!IsValid(ClonerComponent))
 	{
-		return false;
+		return;
 	}
 
-	// Copy template asset since we will modify it, and we do not want other cloner instances to have the new modified asset
-	FObjectDuplicationParameters Parameters = InitStaticDuplicateObjectParams(TemplateNiagaraSystem, GetTransientPackage(), NAME_None, RF_Transient, nullptr, EDuplicateMode::Normal);
-	NiagaraSystem = Cast<UNiagaraSystem>(StaticDuplicateObjectEx(Parameters));
+	// Extract package path
+	FString PackagePath = LayoutAssetPath;
 
-	if (!NiagaraSystem)
-	{
-		return false;
-	}
+	int32 FirstQuoteIndex;
+	PackagePath.FindChar('\'', FirstQuoteIndex);
+
+	int32 LastQuoteIndex;
+	PackagePath.FindLastChar('\'', LastQuoteIndex);
 
-	for (FNiagaraEmitterHandle& SystemEmitterHandle : NiagaraSystem->GetEmitterHandles())
+	if (FirstQuoteIndex != LastQuoteIndex)
 	{
-		if (const FVersionedNiagaraEmitterData* EmitterData = SystemEmitterHandle.GetEmitterData())
-		{
-			for (UNiagaraRendererProperties* EmitterRenderer : EmitterData->GetRenderers())
-			{
-				if (UNiagaraMeshRendererProperties* EmitterMeshRenderer = Cast<UNiagaraMeshRendererProperties>(EmitterRenderer))
-				{
-					EmitterMeshRenderer->Meshes.Empty();
-#if WITH_EDITORONLY_DATA
-					EmitterMeshRenderer->OnMeshChanged();
-#endif
+		PackagePath = PackagePath.Mid(FirstQuoteIndex + 1, LastQuoteIndex - FirstQuoteIndex - 1);
+	}
 
-					MeshRenderer = EmitterMeshRenderer;
-					DataInterfaces = FCEClonerEffectorDataInterfaces(NiagaraSystem);
+	PackagePath = FPackageName::ObjectPathToPackageName(PackagePath);
 
-					OnLayoutLoaded();
+	FLoadPackageAsyncOptionalParams Params;
+	Params.CustomPackageName = FName(TEXT("/") + FString::FromInt(ClonerComponent->GetUniqueID()) + TEXT("_") + GetLayoutName().ToString());
+	Params.CompletionDelegate = MakeUnique<FLoadPackageAsyncDelegate>(FLoadPackageAsyncDelegate::CreateUObject(this, &UCEClonerLayoutBase::OnSystemPackageLoaded));
 
-					return true;
-				}
-			}
-		}
-	}
+	UE_LOG(LogCEClonerLayoutBase, Verbose, TEXT("Cloner layout %s : Template system (%s) loading package (%s) async..."), *LayoutName.ToString(), *LayoutAssetPath, *Params.CustomPackageName.ToString())
 
-	return false;
+	LoadRequestIdentifier = LoadPackageAsync(PackagePath, MoveTemp(Params));
 }
 
 bool UCEClonerLayoutBase::UnloadLayout()
@@ -230,33 +222,79 @@ bool UCEClonerLayoutBase::CopyTo(UCEClonerLayoutBase* InOtherLayout) const
 	return true;
 }
 
+TSet<FName> UCEClonerLayoutBase::GetSupportedExtensions() const
+{
+	TSet<FName> ModuleSupported;
+
+	if (UCEClonerSubsystem* ClonerSubsystem = UCEClonerSubsystem::Get())
+	{
+		ModuleSupported.Append(ClonerSubsystem->GetExtensionNames());
+	}
+
+	return ModuleSupported;
+}
+
 void UCEClonerLayoutBase::OnLayoutPropertyChanged()
 {
 	UpdateLayoutParameters();
 }
 
-UNiagaraSystem* UCEClonerLayoutBase::LoadSystemPath(const FString& InPath) const
+void UCEClonerLayoutBase::OnSystemPackageLoaded(const FName& InName, UPackage* InPackage, EAsyncLoadingResult::Type InResult)
 {
-	UNiagaraSystem* LoadedNiagaraSystem = FindObject<UNiagaraSystem>(nullptr, *InPath);
+	NiagaraSystem = InPackage ? Cast<UNiagaraSystem>(InPackage->FindAssetInPackage()) : nullptr;
+	LoadRequestIdentifier = INDEX_NONE;
 
-	if (!LoadedNiagaraSystem)
+	if (NiagaraSystem)
 	{
-		LoadedNiagaraSystem = LoadObject<UNiagaraSystem>(nullptr, *InPath);
+		// Change outer to avoid GC leak
+		InPackage->SetFlags(RF_Transient);
+		NiagaraSystem->RemoveFromRoot();
+		NiagaraSystem->Rename(NULL, this, REN_DoNotDirty | REN_DontCreateRedirectors | REN_ForceNoResetLoaders | REN_NonTransactional);
+		InPackage->MarkAsGarbage();
+
+		for (FNiagaraEmitterHandle& SystemEmitterHandle : NiagaraSystem->GetEmitterHandles())
+		{
+			if (const FVersionedNiagaraEmitterData* EmitterData = SystemEmitterHandle.GetEmitterData())
+			{
+				for (UNiagaraRendererProperties* EmitterRenderer : EmitterData->GetRenderers())
+				{
+					if (UNiagaraMeshRendererProperties* EmitterMeshRenderer = Cast<UNiagaraMeshRendererProperties>(EmitterRenderer))
+					{
+						EmitterMeshRenderer->Meshes.Empty();
+#if WITH_EDITORONLY_DATA
+						EmitterMeshRenderer->OnMeshChanged();
+#endif
+
+						MeshRenderer = EmitterMeshRenderer;
+						DataInterfaces = FCEClonerEffectorDataInterfaces(NiagaraSystem);
+
+						UE_LOG(LogCEClonerLayoutBase, Log, TEXT("Cloner layout %s : Template system (%s) package (%s) loaded"), *LayoutName.ToString(), *LayoutAssetPath, *InName.ToString())
+
+						OnLayoutLoaded();
+					}
+				}
+			}
+		}
+	}
+	else
+	{
+		UE_LOG(LogCEClonerLayoutBase, Warning, TEXT("Cloner layout %s : Template system (%s) package (%s) could not be async loaded"), *LayoutName.ToString(), *LayoutAssetPath, *InName.ToString())
 	}
 
-	return LoadedNiagaraSystem;
+	OnClonerLayoutLoadedDelegate.Broadcast(this, IsLayoutLoaded());
+	OnClonerLayoutLoadedDelegate.Clear();
 }
 
-ACEClonerActor* UCEClonerLayoutBase::GetClonerActor() const
+UCEClonerComponent* UCEClonerLayoutBase::GetClonerComponent() const
 {
-	return GetTypedOuter<ACEClonerActor>();
+	return GetTypedOuter<UCEClonerComponent>();
 }
 
-UCEClonerComponent* UCEClonerLayoutBase::GetClonerComponent() const
+AActor* UCEClonerLayoutBase::GetClonerActor() const
 {
-	if (const ACEClonerActor* Cloner = GetClonerActor())
+	if (UCEClonerComponent* ClonerComponent = GetClonerComponent())
 	{
-		return Cloner->GetClonerComponent();
+		return ClonerComponent->GetOwner();
 	}
 
 	return nullptr;
@@ -282,8 +320,8 @@ void UCEClonerLayoutBase::UpdateLayoutParameters(bool bInUpdateCloner, bool bInI
 
 void UCEClonerLayoutBase::RequestClonerUpdate(bool bInImmediate) const
 {
-	if (ACEClonerActor* Cloner = GetClonerActor())
+	if (UCEClonerComponent* ClonerComponent = GetClonerComponent())
 	{
-		Cloner->RequestClonerUpdate(bInImmediate);
+		ClonerComponent->RequestClonerUpdate(bInImmediate);
 	}
 }
diff --git a/Engine/Plugins/Experimental/ClonerEffector/Source/ClonerEffector/Private/Cloner/Layouts/CEClonerSplineLayout.cpp b/Engine/Plugins/Experimental/ClonerEffector/Source/ClonerEffector/Private/Cloner/Layouts/CEClonerSplineLayout.cpp
--- a/Engine/Plugins/Experimental/ClonerEffector/Source/ClonerEffector/Private/Cloner/Layouts/CEClonerSplineLayout.cpp
+++ b/Engine/Plugins/Experimental/ClonerEffector/Source/ClonerEffector/Private/Cloner/Layouts/CEClonerSplineLayout.cpp
@@ -2,7 +2,6 @@
 
 #include "Cloner/Layouts/CEClonerSplineLayout.h"
 
-#include "Cloner/CEClonerActor.h"
 #include "Cloner/CEClonerComponent.h"
 #include "Engine/World.h"
 #include "GameFramework/Actor.h"
@@ -54,14 +53,19 @@ void UCEClonerSplineLayout::SetOrientMesh(bool bInOrientMesh)
 #if WITH_EDITOR
 void UCEClonerSplineLayout::SpawnLinkedSplineActor()
 {
-	const ACEClonerActor* ClonerActor = GetClonerActor();
+	const UCEClonerComponent* ClonerComponent = GetClonerComponent();
 
-	if (!ClonerActor)
+	if (!IsValid(ClonerComponent))
 	{
 		return;
 	}
 
-	UWorld* ClonerWorld = ClonerActor->GetWorld();
+	UWorld* ClonerWorld = ClonerComponent->GetWorld();
+
+	if (!IsValid(ClonerWorld))
+	{
+		return;
+	}
 
 	FActorSpawnParameters Params;
 	Params.bTemporaryEditorActor = false;
@@ -89,8 +93,8 @@ void UCEClonerSplineLayout::SpawnLinkedSplineActor()
 	// Rerun construction scripts
 	SpawnedSplineActor->RerunConstructionScripts();
 
-	SpawnedSplineActor->SetActorLocation(ClonerActor->GetActorLocation());
-	SpawnedSplineActor->SetActorRotation(ClonerActor->GetActorRotation());
+	SpawnedSplineActor->SetActorLocation(ClonerComponent->GetComponentLocation());
+	SpawnedSplineActor->SetActorRotation(ClonerComponent->GetComponentRotation());
 
 	SetSplineActorWeak(SpawnedSplineActor);
 	FActorLabelUtilities::RenameExistingActor(SpawnedSplineActor, TEXT("SplineActor"), true);
@@ -160,7 +164,7 @@ void UCEClonerSplineLayout::OnLayoutParametersChanged(UCEClonerComponent* InComp
 			USceneComponent::MarkRenderStateDirtyEvent.RemoveAll(this);
 			USceneComponent::MarkRenderStateDirtyEvent.AddUObject(this, &UCEClonerSplineLayout::OnSampleSplineRenderStateUpdated);
 
-			if (ACEClonerActor* ClonerActor = GetClonerActor())
+			if (AActor* ClonerActor = GetClonerActor())
 			{
 				ClonerActor->SetActorTransform(SplineActor->GetActorTransform());
 			}
