MotionDesign : ClonerEffector

diff --git a/Engine/Plugins/Experimental/ClonerEffector/Source/ClonerEffector/Private/Effector/CEEffectorActor.cpp b/Engine/Plugins/Experimental/ClonerEffector/Source/ClonerEffector/Private/Effector/CEEffectorActor.cpp
--- a/Engine/Plugins/Experimental/ClonerEffector/Source/ClonerEffector/Private/Effector/CEEffectorActor.cpp
+++ b/Engine/Plugins/Experimental/ClonerEffector/Source/ClonerEffector/Private/Effector/CEEffectorActor.cpp
@@ -2,1463 +2,230 @@
 
 #include "Effector/CEEffectorActor.h"
 
-#include "Cloner/CEClonerActor.h"
 #include "Components/DynamicMeshComponent.h"
 #include "Effector/CEEffectorComponent.h"
-#include "Math/Vector.h"
-#include "Subsystems/CEEffectorSubsystem.h"
+#include "Effector/Extensions/CEEffectorForceExtension.h"
+#include "Effector/Modes/CEEffectorNoiseMode.h"
+#include "Effector/Modes/CEEffectorOffsetMode.h"
+#include "Effector/Modes/CEEffectorPushMode.h"
+#include "Effector/Modes/CEEffectorTargetMode.h"
+#include "Effector/Types/CEEffectorBoxType.h"
+#include "Effector/Types/CEEffectorPlaneType.h"
+#include "Effector/Types/CEEffectorRadialType.h"
+#include "Effector/Types/CEEffectorSphereType.h"
+#include "Effector/Types/CEEffectorTorusType.h"
 
 #if WITH_EDITOR
-#include "GeometryScript/MeshPrimitiveFunctions.h"
 #include "Materials/MaterialInstanceDynamic.h"
-#include "PropertyBag.h"
-#include "Settings/CEClonerEffectorSettings.h"
 #include "UObject/ConstructorHelpers.h"
 #endif
 
 DEFINE_LOG_CATEGORY_STATIC(LogCEEffectorActor, Log, All);
 
-ACEEffectorActor::ACEEffectorActor()
-{
-	SetCanBeDamaged(false);
-	PrimaryActorTick.bCanEverTick = false;
-
-	USceneComponent* SceneComponent = CreateDefaultSubobject<UCEEffectorComponent>(TEXT("AvaEffectorComponent"));
-#if WITH_EDITORONLY_DATA
-	SceneComponent->bVisualizeComponent = true;
-#endif
-	SetRootComponent(SceneComponent);
-
-	// Inner Visualizer Component
-	InnerVisualizerComponent = CreateDefaultSubobject<UDynamicMeshComponent>(TEXT("InnerVisualizerComponent"));
-	InnerVisualizerComponent->SetHiddenInGame(true);
-	InnerVisualizerComponent->SetTranslucentSortPriority(0);
-#if WITH_EDITOR
-	InnerVisualizerComponent->SetIsVisualizationComponent(true);
-#endif
-	InnerVisualizerComponent->bIsEditorOnly = true;
-	InnerVisualizerComponent->SetupAttachment(SceneComponent);
-
-	// Outer Visualizer Component
-	OuterVisualizerComponent = CreateDefaultSubobject<UDynamicMeshComponent>(TEXT("OuterVisualizerComponent"));
-	OuterVisualizerComponent->SetHiddenInGame(true);
-	OuterVisualizerComponent->SetTranslucentSortPriority(1);
-#if WITH_EDITOR
-	OuterVisualizerComponent->SetIsVisualizationComponent(true);
-#endif
-	OuterVisualizerComponent->bIsEditorOnly = true;
-	OuterVisualizerComponent->SetupAttachment(SceneComponent);
-
-#if WITH_EDITOR
-	static const ConstructorHelpers::FObjectFinder<UMaterialInterface> VisualizerMaterial(TEXT("/Script/Engine.Material'/ClonerEffector/Materials/M_EffectorVisualizer.M_EffectorVisualizer'"));
-
-	if (VisualizerMaterial.Succeeded())
-	{
-		InnerVisualizerMaterial = UMaterialInstanceDynamic::Create(
-			VisualizerMaterial.Object,
-			this
-		);
-
-		OuterVisualizerMaterial = UMaterialInstanceDynamic::Create(
-			VisualizerMaterial.Object,
-			this
-		);
-	}
-#endif
-
-	if (!IsTemplate())
-	{
-		SceneComponent->TransformUpdated.AddUObject(this, &ACEEffectorActor::OnEffectorTransformed);
-		UCEEffectorSubsystem::OnSubsystemInitializedDelegate.AddUObject(this, &ACEEffectorActor::OnEffectorSubsystemInitialized);
-
-#if WITH_EDITOR
-		if (UCEClonerEffectorSettings* ClonerEffectorSettings = GetMutableDefault<UCEClonerEffectorSettings>())
-		{
-			ClonerEffectorSettings->OnSettingChanged().AddUObject(this, &ACEEffectorActor::OnEffectorDeveloperSettingsChanged);
-		}
-#endif
-	}
-}
-
-#if WITH_EDITOR
-FString ACEEffectorActor::GetDefaultActorLabel() const
-{
-	return DefaultLabel;
-}
-
-TCEPropertyChangeDispatcher<ACEEffectorActor> ACEEffectorActor::PropertyChangeDispatcher =
-{
-	/** Effector */
-	{ GET_MEMBER_NAME_CHECKED(ACEEffectorActor, bEnabled), &ACEEffectorActor::OnEnabledChanged },
-	{ GET_MEMBER_NAME_CHECKED(ACEEffectorActor, Magnitude), &ACEEffectorActor::OnMagnitudeChanged },
-	{ GET_MEMBER_NAME_CHECKED(ACEEffectorActor, bVisualizerComponentVisible), &ACEEffectorActor::OnVisualizerOptionsChanged },
-	{ GET_MEMBER_NAME_CHECKED(ACEEffectorActor, bVisualizerSpriteVisible), &ACEEffectorActor::OnVisualizerOptionsChanged },
-	{ GET_MEMBER_NAME_CHECKED(ACEEffectorActor, Color), &ACEEffectorActor::OnColorChanged },
-	/** Type */
-	{ GET_MEMBER_NAME_CHECKED(ACEEffectorActor, Type), &ACEEffectorActor::OnTypeChanged },
-	{ GET_MEMBER_NAME_CHECKED(ACEEffectorActor, bInvertType), &ACEEffectorActor::OnMagnitudeChanged },
-	{ GET_MEMBER_NAME_CHECKED(ACEEffectorActor, Easing), &ACEEffectorActor::OnEasingChanged },
-	{ GET_MEMBER_NAME_CHECKED(ACEEffectorActor, OuterRadius), &ACEEffectorActor::OnSphereChanged },
-	{ GET_MEMBER_NAME_CHECKED(ACEEffectorActor, InnerRadius), &ACEEffectorActor::OnSphereChanged },
-	{ GET_MEMBER_NAME_CHECKED(ACEEffectorActor, InnerExtent), &ACEEffectorActor::OnBoxChanged },
-	{ GET_MEMBER_NAME_CHECKED(ACEEffectorActor, OuterExtent), &ACEEffectorActor::OnBoxChanged },
-	{ GET_MEMBER_NAME_CHECKED(ACEEffectorActor, PlaneSpacing), &ACEEffectorActor::OnPlaneChanged },
-	{ GET_MEMBER_NAME_CHECKED(ACEEffectorActor, RadialAngle), &ACEEffectorActor::OnRadialChanged },
-	{ GET_MEMBER_NAME_CHECKED(ACEEffectorActor, RadialMinRadius), &ACEEffectorActor::OnRadialChanged },
-	{ GET_MEMBER_NAME_CHECKED(ACEEffectorActor, RadialMaxRadius), &ACEEffectorActor::OnRadialChanged },
-	{ GET_MEMBER_NAME_CHECKED(ACEEffectorActor, TorusRadius), &ACEEffectorActor::OnTorusChanged },
-	{ GET_MEMBER_NAME_CHECKED(ACEEffectorActor, TorusInnerRadius), &ACEEffectorActor::OnTorusChanged },
-	{ GET_MEMBER_NAME_CHECKED(ACEEffectorActor, TorusOuterRadius), &ACEEffectorActor::OnTorusChanged },
-	/** Mode */
-	{ GET_MEMBER_NAME_CHECKED(ACEEffectorActor, Mode), &ACEEffectorActor::OnModeChanged },
-	{ GET_MEMBER_NAME_CHECKED(ACEEffectorActor, Offset), &ACEEffectorActor::OnTransformOptionsChanged },
-	{ GET_MEMBER_NAME_CHECKED(ACEEffectorActor, Rotation), &ACEEffectorActor::OnTransformOptionsChanged },
-	{ GET_MEMBER_NAME_CHECKED(ACEEffectorActor, Scale), &ACEEffectorActor::OnTransformOptionsChanged },
-	{ GET_MEMBER_NAME_CHECKED(ACEEffectorActor, TargetActorWeak), &ACEEffectorActor::OnTargetActorChanged },
-	{ GET_MEMBER_NAME_CHECKED(ACEEffectorActor, LocationStrength), &ACEEffectorActor::OnNoiseFieldOptionsChanged },
-	{ GET_MEMBER_NAME_CHECKED(ACEEffectorActor, RotationStrength), &ACEEffectorActor::OnNoiseFieldOptionsChanged },
-	{ GET_MEMBER_NAME_CHECKED(ACEEffectorActor, ScaleStrength), &ACEEffectorActor::OnNoiseFieldOptionsChanged },
-	{ GET_MEMBER_NAME_CHECKED(ACEEffectorActor, Pan), &ACEEffectorActor::OnNoiseFieldOptionsChanged },
-	{ GET_MEMBER_NAME_CHECKED(ACEEffectorActor, Frequency), &ACEEffectorActor::OnNoiseFieldOptionsChanged },
-	{ GET_MEMBER_NAME_CHECKED(ACEEffectorActor, PushStrength), &ACEEffectorActor::OnPushOptionsChanged },
-	{ GET_MEMBER_NAME_CHECKED(ACEEffectorActor, PushDirection), &ACEEffectorActor::OnPushOptionsChanged },
-	/** Force */
-	{ GET_MEMBER_NAME_CHECKED(ACEEffectorActor, bOrientationForceEnabled), &ACEEffectorActor::OnForceEnabledChanged },
-	{ GET_MEMBER_NAME_CHECKED(ACEEffectorActor, OrientationForceRate), &ACEEffectorActor::OnForceOptionsChanged },
-	{ GET_MEMBER_NAME_CHECKED(ACEEffectorActor, OrientationForceMin), &ACEEffectorActor::OnForceOptionsChanged },
-	{ GET_MEMBER_NAME_CHECKED(ACEEffectorActor, OrientationForceMax), &ACEEffectorActor::OnForceOptionsChanged },
-	{ GET_MEMBER_NAME_CHECKED(ACEEffectorActor, bVortexForceEnabled), &ACEEffectorActor::OnForceEnabledChanged },
-	{ GET_MEMBER_NAME_CHECKED(ACEEffectorActor, VortexForceAmount), &ACEEffectorActor::OnForceOptionsChanged },
-	{ GET_MEMBER_NAME_CHECKED(ACEEffectorActor, VortexForceAxis), &ACEEffectorActor::OnForceOptionsChanged },
-	{ GET_MEMBER_NAME_CHECKED(ACEEffectorActor, bCurlNoiseForceEnabled), &ACEEffectorActor::OnForceEnabledChanged },
-	{ GET_MEMBER_NAME_CHECKED(ACEEffectorActor, CurlNoiseForceStrength), &ACEEffectorActor::OnForceOptionsChanged },
-	{ GET_MEMBER_NAME_CHECKED(ACEEffectorActor, CurlNoiseForceFrequency), &ACEEffectorActor::OnForceOptionsChanged },
-	{ GET_MEMBER_NAME_CHECKED(ACEEffectorActor, bAttractionForceEnabled), &ACEEffectorActor::OnForceEnabledChanged },
-	{ GET_MEMBER_NAME_CHECKED(ACEEffectorActor, AttractionForceStrength), &ACEEffectorActor::OnForceOptionsChanged },
-	{ GET_MEMBER_NAME_CHECKED(ACEEffectorActor, AttractionForceFalloff), &ACEEffectorActor::OnForceOptionsChanged },
-	{ GET_MEMBER_NAME_CHECKED(ACEEffectorActor, bGravityForceEnabled), &ACEEffectorActor::OnForceEnabledChanged },
-	{ GET_MEMBER_NAME_CHECKED(ACEEffectorActor, GravityForceAcceleration), &ACEEffectorActor::OnForceOptionsChanged },
-};
-
-void ACEEffectorActor::PostEditChangeProperty(FPropertyChangedEvent& InPropertyChangedEvent)
-{
-	Super::PostEditChangeProperty(InPropertyChangedEvent);
-
-	PropertyChangeDispatcher.OnPropertyChanged(this, InPropertyChangedEvent);
-}
-
-void ACEEffectorActor::PostEditImport()
-{
-	Super::PostEditImport();
-
-	OnEffectorChanged();
-}
-
-void ACEEffectorActor::PostEditUndo()
-{
-	Super::PostEditUndo();
-
-	OnEffectorChanged();
-}
-#endif
-
-void ACEEffectorActor::Destroyed()
-{
-	Super::Destroyed();
-
-	// Remove this effector from the effector channel
-	if (UCEEffectorSubsystem* EffectorSubsystem = UCEEffectorSubsystem::Get(GetWorld()))
-	{
-		EffectorSubsystem->UnregisterChannelEffector(this);
-	}
-}
-
-void ACEEffectorActor::PostActorCreated()
-{
-	Super::PostActorCreated();
-
-	OnEffectorChanged();
-}
-
-void ACEEffectorActor::PostLoad()
-{
-	Super::PostLoad();
-
-	OnEffectorChanged();
-}
-
-void ACEEffectorActor::PostDuplicate(EDuplicateMode::Type InDuplicateMode)
-{
-	Super::PostDuplicate(InDuplicateMode);
-
-	for (const TWeakObjectPtr<ACEClonerActor>& ClonerWeak : ClonersWeak)
-	{
-		if (ACEClonerActor* Cloner = ClonerWeak.Get())
-		{
-			UE_LOG(LogCEEffectorActor, Log, TEXT("Linking effector %s to cloner %s after duplication"), *GetActorNameOrLabel(), *Cloner->GetActorNameOrLabel())
-			Cloner->LinkEffector(this);
-		}
-	}
-
-	OnEffectorChanged();
-}
-
-void ACEEffectorActor::RegisterToChannel()
-{
-	if (IsValid(this) && ChannelData.GetIdentifier() == INDEX_NONE)
-	{
-		// Register this effector to the effector channel
-		if (UCEEffectorSubsystem* EffectorSubsystem = UCEEffectorSubsystem::Get(GetWorld()))
-		{
-			EffectorSubsystem->RegisterChannelEffector(this);
-		}
-	}
-}
-
-int32 ACEEffectorActor::GetChannelIdentifier() const
-{
-	return ChannelData.GetIdentifier();
-}
-
-void ACEEffectorActor::OnEffectorSubsystemInitialized(const UWorld* InWorld)
-{
-	if (InWorld && GetWorld() == InWorld)
-	{
-		OnEffectorChanged();
-	}
-}
-
-void ACEEffectorActor::OnClonerLinked(ACEClonerActor* InCloner)
-{
-	if (!IsValid(InCloner) || ClonersWeak.Contains(InCloner))
-	{
-		return;
-	}
-
-	ClonersWeak.Add(InCloner);
-}
-
-void ACEEffectorActor::OnClonerUnlinked(ACEClonerActor* InCloner)
-{
-	if (!IsValid(InCloner) || !ClonersWeak.Contains(InCloner))
-	{
-		return;
-	}
-
-	ClonersWeak.Remove(InCloner);
-}
-
-FCEClonerEffectorChannelData& ACEEffectorActor::GetChannelData()
-{
-	return ChannelData;
-}
-
-void ACEEffectorActor::SetType(ECEClonerEffectorType InType)
-{
-	if (InType == Type)
-	{
-		return;
-	}
-
-	Type = InType;
-	OnTypeChanged();
-}
-
-void ACEEffectorActor::SetEasing(ECEClonerEasing InEasing)
-{
-	if (InEasing == Easing)
-	{
-		return;
-	}
-
-	Easing = InEasing;
-	OnEasingChanged();
-}
-
-void ACEEffectorActor::SetMagnitude(float InMagnitude)
-{
-	if (FMath::IsNearlyEqual(InMagnitude, Magnitude))
-	{
-		return;
-	}
-
-	if (InMagnitude < 0.f || InMagnitude > 1.f)
-	{
-		return;
-	}
-
-	Magnitude = InMagnitude;
-	OnMagnitudeChanged();
-}
-
-void ACEEffectorActor::SetOuterRadius(float InRadius)
-{
-	if (FMath::IsNearlyEqual(InRadius, OuterRadius))
-	{
-		return;
-	}
-
-	if (InRadius < 0.f || InRadius < InnerRadius)
-	{
-		return;
-	}
-
-	OuterRadius = InRadius;
-	OnSphereChanged();
-}
-
-void ACEEffectorActor::SetInnerRadius(float InRadius)
-{
-	if (FMath::IsNearlyEqual(InRadius, InnerRadius))
-	{
-		return;
-	}
-
-	if (InRadius < 0.f || InRadius > OuterRadius)
-	{
-		return;
-	}
-
-	InnerRadius = InRadius;
-	OnSphereChanged();
-}
-
-void ACEEffectorActor::SetInnerExtent(const FVector& InExtent)
-{
-	if (InExtent.Equals(InnerExtent))
-	{
-		return;
-	}
-
-	if (InExtent.X < 0 || InExtent.Y < 0 || InExtent.Z < 0)
-	{
-		return;
-	}
-
-	InnerExtent = InExtent;
-	OnBoxChanged();
-}
-
-void ACEEffectorActor::SetOuterExtent(const FVector& InExtent)
-{
-	if (InExtent.Equals(OuterExtent))
-	{
-		return;
-	}
-
-	if (InExtent.X < 0 || InExtent.Y < 0 || InExtent.Z < 0)
-	{
-		return;
-	}
-
-	OuterExtent = InExtent;
-	OnBoxChanged();
-}
-
-void ACEEffectorActor::SetPlaneSpacing(float InSpacing)
-{
-	if (FMath::IsNearlyEqual(InSpacing, PlaneSpacing))
-	{
-		return;
-	}
-
-	if (InSpacing < 0.f)
-	{
-		return;
-	}
-
-	PlaneSpacing = InSpacing;
-	OnPlaneChanged();
-}
-
-void ACEEffectorActor::SetRadialAngle(float InAngle)
-{
-	InAngle = FMath::Clamp(InAngle, 0, 360);
-
-	if (FMath::IsNearlyEqual(InAngle, RadialAngle))
-	{
-		return;
-	}
-
-	RadialAngle = InAngle;
-	OnRadialChanged();
-}
-
-void ACEEffectorActor::SetRadialMinRadius(float InRadius)
-{
-	InRadius = FMath::Max(0, InRadius);
-
-	if (FMath::IsNearlyEqual(InRadius, RadialMinRadius))
-	{
-		return;
-	}
-
-	RadialMinRadius = InRadius;
-	OnRadialChanged();
-}
-
-void ACEEffectorActor::SetInvertType(bool bInInvert)
-{
-	if (bInvertType == bInInvert)
-	{
-		return;
-	}
-
-	bInvertType = bInInvert;
-	OnMagnitudeChanged();
-}
-
-void ACEEffectorActor::SetRadialMaxRadius(float InRadius)
-{
-	InRadius = FMath::Max(0, InRadius);
-
-	if (FMath::IsNearlyEqual(InRadius, RadialMaxRadius))
-	{
-		return;
-	}
-
-	RadialMaxRadius = InRadius;
-	OnRadialChanged();
-}
-
-void ACEEffectorActor::SetTorusRadius(float InRadius)
-{
-	InRadius = FMath::Max(0, InRadius);
-
-	if (FMath::IsNearlyEqual(InRadius, TorusRadius))
-	{
-		return;
-	}
-
-	TorusRadius = InRadius;
-	OnTorusChanged();
-}
-
-void ACEEffectorActor::SetTorusInnerRadius(float InRadius)
-{
-	InRadius = FMath::Max(0, InRadius);
-
-	if (FMath::IsNearlyEqual(InRadius, TorusInnerRadius))
-	{
-		return;
-	}
-
-	TorusInnerRadius = InRadius;
-	OnTorusChanged();
-}
-
-void ACEEffectorActor::SetTorusOuterRadius(float InRadius)
-{
-	InRadius = FMath::Max(0, InRadius);
-
-	if (FMath::IsNearlyEqual(InRadius, TorusOuterRadius))
-	{
-		return;
-	}
-
-	TorusOuterRadius = InRadius;
-	OnTorusChanged();
-}
-
-void ACEEffectorActor::SetOffset(const FVector& InOffset)
-{
-	if (InOffset.Equals(Offset))
-	{
-		return;
-	}
-
-	Offset = InOffset;
-	OnTransformOptionsChanged();
-}
-
-void ACEEffectorActor::SetRotation(const FRotator& InRotation)
-{
-	if (InRotation.Equals(Rotation))
-	{
-		return;
-	}
-
-	constexpr float MinRotation = -180.f;
-	if (InRotation.Pitch < MinRotation || InRotation.Roll < MinRotation || InRotation.Yaw < MinRotation)
-	{
-		return;
-	}
-
-	constexpr float MaxRotation = 180.f;
-	if (InRotation.Pitch > MaxRotation || InRotation.Roll > MaxRotation || InRotation.Yaw > MaxRotation)
-	{
-		return;
-	}
-
-	Rotation = InRotation;
-	OnTransformOptionsChanged();
-}
-
-void ACEEffectorActor::SetScale(const FVector& InScale)
-{
-	if (InScale.Equals(Scale))
-	{
-		return;
-	}
-
-	if (InScale.X < 0.f || InScale.Y < 0.f || InScale.Z < 0.f)
-	{
-		return;
-	}
-
-	Scale = InScale;
-	OnTransformOptionsChanged();
-}
-
-void ACEEffectorActor::SetEnabled(bool bInEnable)
-{
-	if (bInEnable == bEnabled)
-	{
-		return;
-	}
-
-	bEnabled = bInEnable;
-	OnEnabledChanged();
-}
-
-void ACEEffectorActor::SetMode(ECEClonerEffectorMode InMode)
-{
-	if (Mode == InMode)
-	{
-		return;
-	}
-
-	Mode = InMode;
-	OnModeChanged();
-}
-
-void ACEEffectorActor::SetTargetActor(AActor* InTargetActor)
-{
-	if (InTargetActor == TargetActorWeak.Get())
-	{
-		return;
-	}
-
-	TargetActorWeak = InTargetActor;
-	OnTargetActorChanged();
-}
-
-void ACEEffectorActor::SetTargetActorWeak(const TWeakObjectPtr<AActor>& InTargetActor)
-{
-	SetTargetActor(InTargetActor.Get());
-}
-
-void ACEEffectorActor::SetAttractionForceStrength(float InForceStrength)
-{
-	if (FMath::IsNearlyEqual(AttractionForceStrength, InForceStrength))
-	{
-		return;
-	}
-
-	AttractionForceStrength = InForceStrength;
-	OnForceOptionsChanged();
-}
-
-void ACEEffectorActor::SetAttractionForceFalloff(float InForceFalloff)
-{
-	if (FMath::IsNearlyEqual(AttractionForceFalloff, InForceFalloff))
-	{
-		return;
-	}
-
-	AttractionForceFalloff = InForceFalloff;
-	OnForceOptionsChanged();
-}
-
-void ACEEffectorActor::SetGravityForceEnabled(bool bInForceEnabled)
-{
-	if (bGravityForceEnabled == bInForceEnabled)
-	{
-		return;
-	}
-
-	bGravityForceEnabled = bInForceEnabled;
-	OnForceEnabledChanged();
-}
-
-void ACEEffectorActor::SetGravityForceAcceleration(const FVector& InAcceleration)
-{
-	if (GravityForceAcceleration.Equals(InAcceleration))
-	{
-		return;
-	}
-
-	GravityForceAcceleration = InAcceleration;
-	OnForceOptionsChanged();
-}
-
-void ACEEffectorActor::SetLocationStrength(const FVector& InStrength)
-{
-	if (LocationStrength.Equals(InStrength))
-	{
-		return;
-	}
-
-	LocationStrength = InStrength;
-	OnNoiseFieldOptionsChanged();
-}
-
-void ACEEffectorActor::SetRotationStrength(const FRotator& InStrength)
-{
-	if (RotationStrength.Equals(InStrength))
-	{
-		return;
-	}
-
-	RotationStrength = InStrength;
-	OnNoiseFieldOptionsChanged();
-}
-
-void ACEEffectorActor::SetScaleStrength(const FVector& InStrength)
-{
-	if (ScaleStrength.Equals(InStrength))
-	{
-		return;
-	}
-
-	ScaleStrength = InStrength;
-	OnNoiseFieldOptionsChanged();
-}
-
-void ACEEffectorActor::SetPan(const FVector& InPan)
-{
-	if (Pan.Equals(InPan))
-	{
-		return;
-	}
-
-	Pan = InPan;
-	OnNoiseFieldOptionsChanged();
-}
-
-void ACEEffectorActor::SetFrequency(float InFrequency)
-{
-	if (FMath::IsNearlyEqual(Frequency, InFrequency))
-	{
-		return;
-	}
-
-	if (InFrequency < 0.f)
-	{
-		return;
-	}
-
-	Frequency = InFrequency;
-	OnNoiseFieldOptionsChanged();
-}
-
-void ACEEffectorActor::SetOrientationForceEnabled(bool bInForceEnabled)
-{
-	if (bOrientationForceEnabled == bInForceEnabled)
-	{
-		return;
-	}
-
-	bOrientationForceEnabled = bInForceEnabled;
-	OnForceEnabledChanged();
-}
-
-void ACEEffectorActor::SetOrientationForceRate(float InForceOrientationRate)
-{
-	if (FMath::IsNearlyEqual(OrientationForceRate, InForceOrientationRate))
-	{
-		return;
-	}
-
-	OrientationForceRate = InForceOrientationRate;
-	OnForceOptionsChanged();
-}
-
-void ACEEffectorActor::SetOrientationForceMin(const FVector& InForceOrientationMin)
-{
-	if (OrientationForceMin.Equals(InForceOrientationMin))
-	{
-		return;
-	}
-
-	OrientationForceMin = InForceOrientationMin;
-	OnForceOptionsChanged();
-}
-
-void ACEEffectorActor::SetOrientationForceMax(const FVector& InForceOrientationMax)
-{
-	if (OrientationForceMax.Equals(InForceOrientationMax))
-	{
-		return;
-	}
-
-	OrientationForceMax = InForceOrientationMax;
-	OnForceOptionsChanged();
-}
-
-void ACEEffectorActor::SetVortexForceEnabled(bool bInForceEnabled)
-{
-	if (bVortexForceEnabled == bInForceEnabled)
-	{
-		return;
-	}
-
-	bVortexForceEnabled = bInForceEnabled;
-	OnForceEnabledChanged();
-}
-
-void ACEEffectorActor::SetVortexForceAmount(float InForceVortexAmount)
-{
-	if (FMath::IsNearlyEqual(VortexForceAmount, InForceVortexAmount))
-	{
-		return;
-	}
-
-	VortexForceAmount = InForceVortexAmount;
-	OnForceOptionsChanged();
-}
-
-void ACEEffectorActor::SetVortexForceAxis(const FVector& InForceVortexAxis)
-{
-	if (VortexForceAxis.Equals(InForceVortexAxis))
-	{
-		return;
-	}
-
-	VortexForceAxis = InForceVortexAxis;
-	OnForceOptionsChanged();
-}
-
-void ACEEffectorActor::SetCurlNoiseForceEnabled(bool bInForceEnabled)
-{
-	if (bCurlNoiseForceEnabled == bInForceEnabled)
-	{
-		return;
-	}
-
-	bCurlNoiseForceEnabled = bInForceEnabled;
-	OnForceEnabledChanged();
-}
-
-void ACEEffectorActor::SetCurlNoiseForceStrength(float InForceCurlNoiseStrength)
-{
-	if (FMath::IsNearlyEqual(CurlNoiseForceStrength, InForceCurlNoiseStrength))
-	{
-		return;
-	}
-
-	CurlNoiseForceStrength = InForceCurlNoiseStrength;
-	OnForceOptionsChanged();
-}
-
-void ACEEffectorActor::SetCurlNoiseForceFrequency(float InForceCurlNoiseFrequency)
-{
-	if (FMath::IsNearlyEqual(CurlNoiseForceFrequency, InForceCurlNoiseFrequency))
-	{
-		return;
-	}
-
-	CurlNoiseForceFrequency = InForceCurlNoiseFrequency;
-	OnForceOptionsChanged();
-}
-
-void ACEEffectorActor::SetAttractionForceEnabled(bool bInForceEnabled)
-{
-	if (bAttractionForceEnabled == bInForceEnabled)
-	{
-		return;
-	}
-
-	bAttractionForceEnabled = bInForceEnabled;
-	OnForceEnabledChanged();
-}
-
-#if WITH_EDITOR
-void ACEEffectorActor::SetVisualizerComponentVisible(bool bInVisible)
-{
-	if (bVisualizerComponentVisible == bInVisible)
-	{
-		return;
-	}
-
-	bVisualizerComponentVisible = bInVisible;
-	OnVisualizerOptionsChanged();
-}
-
-void ACEEffectorActor::SetVisualizerSpriteVisible(bool bInVisible)
-{
-	if (bVisualizerSpriteVisible == bInVisible)
-	{
-		return;
-	}
-
-	bVisualizerSpriteVisible = bInVisible;
-	OnVisualizerOptionsChanged();
-}
-#endif
-
-void ACEEffectorActor::SetColor(const FLinearColor& InColor)
-{
-	if (Color == InColor)
-	{
-		return;
-	}
-
-	Color = InColor;
-	OnColorChanged();
-}
-
-void ACEEffectorActor::SetPushDirection(ECEClonerEffectorPushDirection InDirection)
-{
-	if (PushDirection == InDirection)
-	{
-		return;
-	}
-
-	PushDirection = InDirection;
-	OnPushOptionsChanged();
-}
-
-void ACEEffectorActor::SetPushStrength(const FVector& InStrength)
-{
-	if (PushStrength.Equals(InStrength))
-	{
-		return;
-	}
-
-	PushStrength = InStrength;
-	OnPushOptionsChanged();
-}
-
-void ACEEffectorActor::OnEffectorTransformed(USceneComponent* InUpdatedComponent, EUpdateTransformFlags InUpdateTransformFlags, ETeleportType InTeleport)
-{
-	OnTransformChanged();
-
-	// Update when scaled or rotated
-	UpdateEffectorTypes();
-
-	// update if self
-	const AActor* InternalTargetActor = InternalTargetActorWeak.Get();
-	if (InternalTargetActor == this)
-	{
-		OnTargetOptionsChanged();
-	}
-}
-
-void ACEEffectorActor::OnEnabledChanged()
+struct FCEEffectorActorVersion
 {
-	if (bEnabled)
-	{
-		OnEffectorChanged();
-	}
-	else // Disabled
+	enum Type : int32
 	{
-		OnEffectorDisabled();
-
-		// Hide visualization components
-		OnTypeChanged();
-
-#if WITH_EDITOR
-		OnVisualizerOptionsChanged();
-#endif
-	}
-}
-
-void ACEEffectorActor::OnEffectorDisabled()
-{
-	OnMagnitudeChanged();
-}
+		PreVersioning = 0,
 
-void ACEEffectorActor::OnModeChanged()
-{
-	ChannelData.Mode = Mode;
+		/** Migrating all logic and properties to component */
+		ComponentMigration,
 
-	OnTransformOptionsChanged();
-	OnTargetActorChanged();
-	OnNoiseFieldOptionsChanged();
-	OnPushOptionsChanged();
-}
+		VersionPlusOne,
+		LatestVersion = VersionPlusOne - 1
+	};
 
-void ACEEffectorActor::OnTypeChanged()
-{
-	ChannelData.Type = Type;
+	static constexpr FGuid GUID = FGuid(0x9271D8A6, 0xBF4146B2, 0xA20FC0A4, 0x9D8295B3);
+};
 
-	// Update type data
-	UpdateEffectorTypes();
-}
+FCustomVersionRegistration GRegisterCEEffectorActorVersion(FCEEffectorActorVersion::GUID, static_cast<int32>(FCEEffectorActorVersion::LatestVersion), TEXT("CEEffectorActorVersion"));
 
-void ACEEffectorActor::OnEffectorChanged()
+ACEEffectorActor::ACEEffectorActor()
 {
-	RegisterToChannel();
+	SetCanBeDamaged(false);
+	PrimaryActorTick.bCanEverTick = false;
 
-	if (!bEnabled)
-	{
-		OnEnabledChanged();
-	}
-	else
-	{
-		OnModeChanged();
-		OnTypeChanged();
-		OnEasingChanged();
-		OnTransformChanged();
-		OnMagnitudeChanged();
-		OnForceOptionsChanged();
-		OnColorChanged();
+	EffectorComponent = CreateDefaultSubobject<UCEEffectorComponent>(TEXT("AvaEffectorComponent"));
+	SetRootComponent(EffectorComponent);
 
+	UDynamicMeshComponent* InnerVisualizerComponent = CreateDefaultSubobject<UDynamicMeshComponent>(TEXT("InnerVisualizerComponent"));
+	InnerVisualizerComponent->SetupAttachment(EffectorComponent);
+	InnerVisualizerComponent->bIsEditorOnly = true;
 #if WITH_EDITOR
-		OnVisualizerOptionsChanged();
+	EffectorComponent->AddVisualizerComponent(InnerVisualizerComponent);
 #endif
-	}
-}
 
-void ACEEffectorActor::OnEasingChanged()
-{
-	ChannelData.Easing = Easing;
+	UDynamicMeshComponent* OuterVisualizerComponent = CreateDefaultSubobject<UDynamicMeshComponent>(TEXT("OuterVisualizerComponent"));
+	OuterVisualizerComponent->SetupAttachment(EffectorComponent);
+	OuterVisualizerComponent->bIsEditorOnly = true;
+#if WITH_EDITOR
+	EffectorComponent->AddVisualizerComponent(OuterVisualizerComponent);
+#endif
 }
 
-void ACEEffectorActor::OnTransformChanged()
+#if WITH_EDITOR
+FString ACEEffectorActor::GetDefaultActorLabel() const
 {
-	ChannelData.Location = GetActorLocation();
-	ChannelData.Rotation = GetActorRotation().Quaternion();
-	ChannelData.Scale = GetActorScale();
+	return DefaultLabel;
 }
+#endif
 
-void ACEEffectorActor::OnBoxChanged()
+void ACEEffectorActor::Serialize(FArchive& InArchive)
 {
-	InnerExtent = ClampVector(InnerExtent, FVector::ZeroVector, OuterExtent);
-	OuterExtent = OuterExtent.ComponentMax(InnerExtent);
-
-	if (Type != ECEClonerEffectorType::Box)
-	{
-		return;
-	}
-
-	ChannelData.InnerExtent = GetInnerExtent();
-	ChannelData.OuterExtent = GetOuterExtent();
+	InArchive.UsingCustomVersion(FCEEffectorActorVersion::GUID);
 
-#if WITH_EDITOR
-	// Update visualizer
-	static const FName BoxInnerExtentName = GET_MEMBER_NAME_CHECKED(ACEEffectorActor, InnerExtent);
-	static const FName BoxOuterExtentName = GET_MEMBER_NAME_CHECKED(ACEEffectorActor, OuterExtent);
-	TValueOrError<FVector*, EPropertyBagResult> BoxInnerExtentValue = VisualizerData.GetValueStruct<FVector>(BoxInnerExtentName);
-	TValueOrError<FVector*, EPropertyBagResult> BoxOuterExtentValue = VisualizerData.GetValueStruct<FVector>(BoxOuterExtentName);
+	Super::Serialize(InArchive);
 
-	if (!BoxInnerExtentValue.HasValue()
-		|| *BoxInnerExtentValue.GetValue() != InnerExtent)
-	{
-		UpdateVisualizer(InnerVisualizerId, [this](UDynamicMesh* InMesh)
-		{
-			constexpr FGeometryScriptPrimitiveOptions PrimitiveOptions;
-			UGeometryScriptLibrary_MeshPrimitiveFunctions::AppendBox(InMesh, PrimitiveOptions, FTransform(FVector(0, 0, -InnerExtent.Z)), InnerExtent.X * 2, InnerExtent.Y * 2, InnerExtent.Z * 2);
-		});
-	}
+	const int32 Version = InArchive.CustomVer(FCEEffectorActorVersion::GUID);
 
-	if (!BoxOuterExtentValue.HasValue()
-		|| *BoxOuterExtentValue.GetValue() != OuterExtent)
+	if (Version < FCEEffectorActorVersion::ComponentMigration)
 	{
-		UpdateVisualizer(OuterVisualizerId, [this](UDynamicMesh* InMesh)
-		{
-			constexpr FGeometryScriptPrimitiveOptions PrimitiveOptions;
-			UGeometryScriptLibrary_MeshPrimitiveFunctions::AppendBox(InMesh, PrimitiveOptions, FTransform(FVector(0, 0, -OuterExtent.Z)), OuterExtent.X * 2, OuterExtent.Y * 2, OuterExtent.Z * 2);
-		});
+		MigrateToVersion = FCEEffectorActorVersion::ComponentMigration;
 	}
-
-	VisualizerData.Reset();
-	VisualizerData.AddProperty(BoxInnerExtentName, EPropertyBagPropertyType::Struct);
-	VisualizerData.AddProperty(BoxOuterExtentName, EPropertyBagPropertyType::Struct);
-	VisualizerData.SetValueStruct(BoxInnerExtentName, InnerExtent);
-	VisualizerData.SetValueStruct(BoxOuterExtentName, OuterExtent);
-#endif
 }
 
-void ACEEffectorActor::OnSphereChanged()
+void ACEEffectorActor::PostLoad()
 {
-	InnerRadius = FMath::Clamp(InnerRadius, 0.f, OuterRadius);
-	OuterRadius = FMath::Max(OuterRadius, InnerRadius);
-
-	if (Type != ECEClonerEffectorType::Sphere)
-	{
-		return;
-	}
-
-	ChannelData.InnerExtent = FVector(GetInnerRadius());
-	ChannelData.OuterExtent = FVector(GetOuterRadius());
-
-#if WITH_EDITOR
-	// Update visualizer
-	static const FName SphereInnerRadiusName = GET_MEMBER_NAME_CHECKED(ACEEffectorActor, InnerRadius);
-	static const FName SphereOuterRadiusName = GET_MEMBER_NAME_CHECKED(ACEEffectorActor, OuterRadius);
-	TValueOrError<float, EPropertyBagResult> SphereInnerRadiusValue = VisualizerData.GetValueFloat(SphereInnerRadiusName);
-	TValueOrError<float, EPropertyBagResult> SphereOuterRadiusValue = VisualizerData.GetValueFloat(SphereOuterRadiusName);
-
-	if (!SphereInnerRadiusValue.HasValue()
-		|| SphereInnerRadiusValue.GetValue() != InnerRadius)
-	{
-		UpdateVisualizer(InnerVisualizerId, [this](UDynamicMesh* InMesh)
-		{
-			constexpr FGeometryScriptPrimitiveOptions PrimitiveOptions;
-			UGeometryScriptLibrary_MeshPrimitiveFunctions::AppendSphereLatLong(InMesh, PrimitiveOptions, FTransform::Identity, InnerRadius);
-		});
-	}
-
-	if (!SphereOuterRadiusValue.HasValue()
-		|| SphereOuterRadiusValue.GetValue() != OuterRadius)
-	{
-		UpdateVisualizer(OuterVisualizerId, [this](UDynamicMesh* InMesh)
-		{
-			constexpr FGeometryScriptPrimitiveOptions PrimitiveOptions;
-			UGeometryScriptLibrary_MeshPrimitiveFunctions::AppendSphereLatLong(InMesh, PrimitiveOptions, FTransform::Identity, OuterRadius);
-		});
-	}
+	Super::PostLoad();
 
-	VisualizerData.Reset();
-	VisualizerData.AddProperty(SphereInnerRadiusName, EPropertyBagPropertyType::Float);
-	VisualizerData.AddProperty(SphereOuterRadiusName, EPropertyBagPropertyType::Float);
-	VisualizerData.SetValueFloat(SphereInnerRadiusName, InnerRadius);
-	VisualizerData.SetValueFloat(SphereOuterRadiusName, OuterRadius);
-#endif
+	MigrateDeprecatedProperties();
 }
 
-void ACEEffectorActor::OnPlaneChanged()
+void ACEEffectorActor::MigrateDeprecatedProperties()
 {
-	if (Type != ECEClonerEffectorType::Plane)
+	if (MigrateToVersion == INDEX_NONE)
 	{
 		return;
 	}
 
-	ChannelData.InnerExtent = FVector::LeftVector;
-	ChannelData.OuterExtent = FVector(PlaneSpacing);
+	UE_LOG(LogCEEffectorActor, Warning, TEXT("%s : Effector migrating from version %i to latest %i"), *GetActorNameOrLabel(), MigrateToVersion, FCEEffectorActorVersion::LatestVersion);
 
-#if WITH_EDITOR
-	// Update visualizer
-	static const FName PlaneSpacingName = GET_MEMBER_NAME_CHECKED(ACEEffectorActor, PlaneSpacing);
-	TValueOrError<float, EPropertyBagResult> PlaneSpacingValue = VisualizerData.GetValueFloat(PlaneSpacingName);
+	PRAGMA_DISABLE_DEPRECATION_WARNINGS
 
-	if (!PlaneSpacingValue.HasValue()
-		|| PlaneSpacingValue.GetValue() != PlaneSpacing)
+	if (MigrateToVersion <= FCEEffectorActorVersion::ComponentMigration)
 	{
-		UpdateVisualizer(InnerVisualizerId, [this](UDynamicMesh* InMesh)
-		{
-			const FVector InnerPlane = FVector::LeftVector * FVector(-PlaneSpacing/2);
-			constexpr FGeometryScriptPrimitiveOptions PrimitiveOptions;
-			UGeometryScriptLibrary_MeshPrimitiveFunctions::AppendRectangleXY(InMesh, PrimitiveOptions, FTransform(FRotator(0, 0, 90), InnerPlane), 250, 250);
-		});
+		EffectorComponent->SetEnabled(bEnabled);
+		EffectorComponent->SetMagnitude(Magnitude);
+		EffectorComponent->SetColor(Color);
 
-		UpdateVisualizer(OuterVisualizerId, [this](UDynamicMesh* InMesh)
 		{
-			const FVector OuterPlane = FVector::LeftVector * FVector(PlaneSpacing/2);
-			constexpr FGeometryScriptPrimitiveOptions PrimitiveOptions;
-			UGeometryScriptLibrary_MeshPrimitiveFunctions::AppendRectangleXY(InMesh, PrimitiveOptions, FTransform(FRotator(0, 0, 90), OuterPlane), 500, 500);
-		});
-	}
-
-	VisualizerData.Reset();
-	VisualizerData.AddProperty(PlaneSpacingName, EPropertyBagPropertyType::Float);
-	VisualizerData.SetValueFloat(PlaneSpacingName, PlaneSpacing);
-#endif
-}
-
-void ACEEffectorActor::OnRadialChanged()
-{
-	RadialMinRadius = FMath::Min(RadialMinRadius, RadialMaxRadius);
-	RadialMaxRadius = FMath::Max(RadialMinRadius, RadialMaxRadius);
-
-	if (Type != ECEClonerEffectorType::Radial)
-	{
-		return;
-	}
+			const TArray<FName> TypeNames = EffectorComponent->GetEffectorTypeNames();
+			const int32 TypeIndex = static_cast<int32>(Type);
+			if (TypeNames.IsValidIndex(TypeIndex))
+			{
+				EffectorComponent->SetTypeName(TypeNames[TypeIndex]);
+			}
+		}
 
-	ChannelData.InnerExtent = FVector::RightVector;
-	ChannelData.OuterExtent = FVector(RadialAngle, RadialMinRadius, RadialMaxRadius);
+		{
+			const TArray<FName> ModeNames = EffectorComponent->GetEffectorModeNames();
+			const int32 ModeIndex = static_cast<int32>(Mode);
+			if (ModeNames.IsValidIndex(ModeIndex))
+			{
+				EffectorComponent->SetModeName(ModeNames[ModeIndex]);
+			}
+		}
 
 #if WITH_EDITOR
-	// Update visualizer
-	static const FName RadialAngleName = GET_MEMBER_NAME_CHECKED(ACEEffectorActor, RadialAngle);
-	static const FName RadialMinRadiusName = GET_MEMBER_NAME_CHECKED(ACEEffectorActor, RadialMinRadius);
-	static const FName RadialMaxRadiusName = GET_MEMBER_NAME_CHECKED(ACEEffectorActor, RadialMaxRadius);
-	TValueOrError<float, EPropertyBagResult> RadialAngleValue = VisualizerData.GetValueFloat(RadialAngleName);
-	TValueOrError<float, EPropertyBagResult> RadialMinRadiusValue = VisualizerData.GetValueFloat(RadialMinRadiusName);
-	TValueOrError<float, EPropertyBagResult> RadialMaxRadiusValue = VisualizerData.GetValueFloat(RadialMaxRadiusName);
+		EffectorComponent->SetVisualizerComponentVisible(bVisualizerComponentVisible);
+		EffectorComponent->SetVisualizerSpriteVisible(bVisualizerSpriteVisible);
+#endif
 
-	if (!RadialMaxRadiusValue.HasValue()
-		|| RadialMaxRadiusValue.GetValue() != RadialMaxRadius
-		|| !RadialMinRadiusValue.HasValue()
-		|| RadialMinRadiusValue.GetValue() != RadialMinRadius
-		|| !RadialAngleValue.HasValue()
-		|| RadialAngleValue.GetValue() != RadialAngle)
-	{
-		UpdateVisualizer(InnerVisualizerId, [this](UDynamicMesh* InMesh)
+		if (UCEEffectorSphereType* SphereExtension = EffectorComponent->FindOrAddExtension<UCEEffectorSphereType>())
 		{
-			constexpr FGeometryScriptPrimitiveOptions PrimitiveOptions;
-			UGeometryScriptLibrary_MeshPrimitiveFunctions::AppendDisc(InMesh, PrimitiveOptions, FTransform(FRotator(0, 90, 0)), RadialMaxRadius, 16, 0, 0, RadialAngle / 2, RadialMinRadius);
-		});
+			SphereExtension->SetInnerRadius(InnerRadius);
+			SphereExtension->SetOuterRadius(OuterRadius);
+			SphereExtension->SetInvertType(bInvertType);
+			SphereExtension->SetEasing(Easing);
+		}
 
-		UpdateVisualizer(OuterVisualizerId, [this](UDynamicMesh* InMesh)
+		if (UCEEffectorBoxType* BoxExtension = EffectorComponent->FindOrAddExtension<UCEEffectorBoxType>())
 		{
-			constexpr FGeometryScriptPrimitiveOptions PrimitiveOptions;
-			UGeometryScriptLibrary_MeshPrimitiveFunctions::AppendDisc(InMesh, PrimitiveOptions, FTransform(FRotator(0, 90, 0)), RadialMaxRadius, 16, 0, RadialAngle / 2, RadialAngle, RadialMinRadius);
-		});
-	}
-
-	VisualizerData.Reset();
-	VisualizerData.AddProperty(RadialAngleName, EPropertyBagPropertyType::Float);
-	VisualizerData.AddProperty(RadialMinRadiusName, EPropertyBagPropertyType::Float);
-	VisualizerData.AddProperty(RadialMaxRadiusName, EPropertyBagPropertyType::Float);
-	VisualizerData.SetValueFloat(RadialAngleName, RadialAngle);
-	VisualizerData.SetValueFloat(RadialMinRadiusName, RadialMinRadius);
-	VisualizerData.SetValueFloat(RadialMaxRadiusName, RadialMaxRadius);
-#endif
-}
-
-void ACEEffectorActor::OnTorusChanged()
-{
-	TorusOuterRadius = FMath::Min(TorusRadius, TorusOuterRadius);
-	TorusInnerRadius = FMath::Min(TorusInnerRadius, TorusOuterRadius);
-	TorusOuterRadius = FMath::Max(TorusInnerRadius, TorusOuterRadius);
-
-	if (Type != ECEClonerEffectorType::Torus)
-	{
-		return;
-	}
-
-	ChannelData.InnerExtent = FVector::ZAxisVector;
-	ChannelData.OuterExtent = FVector(TorusInnerRadius, TorusOuterRadius, TorusRadius);
-
-#if WITH_EDITOR
-	// Update visualizer
-	static const FName TorusRadiusName = GET_MEMBER_NAME_CHECKED(ACEEffectorActor, TorusRadius);
-	static const FName TorusInnerRadiusName = GET_MEMBER_NAME_CHECKED(ACEEffectorActor, TorusInnerRadius);
-	static const FName TorusOuterRadiusName = GET_MEMBER_NAME_CHECKED(ACEEffectorActor, TorusOuterRadius);
-	TValueOrError<float, EPropertyBagResult> TorusRadiusValue = VisualizerData.GetValueFloat(TorusRadiusName);
-	TValueOrError<float, EPropertyBagResult> TorusInnerRadiusValue = VisualizerData.GetValueFloat(TorusInnerRadiusName);
-	TValueOrError<float, EPropertyBagResult> TorusOuterRadiusValue = VisualizerData.GetValueFloat(TorusOuterRadiusName);
+			BoxExtension->SetInnerExtent(InnerExtent);
+			BoxExtension->SetOuterExtent(OuterExtent);
+			BoxExtension->SetInvertType(bInvertType);
+			BoxExtension->SetEasing(Easing);
+		}
 
-	if (!TorusRadiusValue.HasValue()
-		|| TorusRadiusValue.GetValue() != TorusRadius
-		|| !TorusInnerRadiusValue.HasValue()
-		|| TorusInnerRadiusValue.GetValue() != TorusInnerRadius)
-	{
-		UpdateVisualizer(InnerVisualizerId, [this](UDynamicMesh* InMesh)
+		if (UCEEffectorPlaneType* PlaneExtension = EffectorComponent->FindOrAddExtension<UCEEffectorPlaneType>())
 		{
-			constexpr FGeometryScriptPrimitiveOptions PrimitiveOptions;
-			constexpr FGeometryScriptRevolveOptions RevolveOptions;
-			UGeometryScriptLibrary_MeshPrimitiveFunctions::AppendTorus(InMesh, PrimitiveOptions, FTransform(FVector(0, 0, -TorusInnerRadius)), RevolveOptions, TorusRadius, TorusInnerRadius);
-		});
-	}
+			PlaneExtension->SetPlaneSpacing(PlaneSpacing);
+			PlaneExtension->SetInvertType(bInvertType);
+			PlaneExtension->SetEasing(Easing);
+		}
 
-	if (!TorusRadiusValue.HasValue()
-		|| TorusRadiusValue.GetValue() != TorusRadius
-		|| !TorusOuterRadiusValue.HasValue()
-		|| TorusOuterRadiusValue.GetValue() != TorusOuterRadius)
-	{
-		UpdateVisualizer(OuterVisualizerId, [this](UDynamicMesh* InMesh)
+		if (UCEEffectorRadialType* RadialExtension = EffectorComponent->FindOrAddExtension<UCEEffectorRadialType>())
 		{
-			constexpr FGeometryScriptPrimitiveOptions PrimitiveOptions;
-			constexpr FGeometryScriptRevolveOptions RevolveOptions;
-			UGeometryScriptLibrary_MeshPrimitiveFunctions::AppendTorus(InMesh, PrimitiveOptions, FTransform(FVector(0, 0, -TorusOuterRadius)), RevolveOptions, TorusRadius, TorusOuterRadius);
-		});
-	}
-
-	VisualizerData.Reset();
-	VisualizerData.AddProperty(TorusRadiusName, EPropertyBagPropertyType::Float);
-	VisualizerData.AddProperty(TorusInnerRadiusName, EPropertyBagPropertyType::Float);
-	VisualizerData.AddProperty(TorusOuterRadiusName, EPropertyBagPropertyType::Float);
-	VisualizerData.SetValueFloat(TorusRadiusName, TorusRadius);
-	VisualizerData.SetValueFloat(TorusInnerRadiusName, TorusInnerRadius);
-	VisualizerData.SetValueFloat(TorusOuterRadiusName, TorusOuterRadius);
-#endif
-}
-
-void ACEEffectorActor::OnUnboundChanged()
-{
-	if (Type != ECEClonerEffectorType::Unbound)
-	{
-		return;
-	}
-
-#if WITH_EDITOR
-	// Clear visualizers
-	UpdateVisualizer(InnerVisualizerId, [](UDynamicMesh* InMesh){});
-	UpdateVisualizer(OuterVisualizerId, [](UDynamicMesh* InMesh){});
-#endif
-}
-
-void ACEEffectorActor::OnMagnitudeChanged()
-{
-	const float EffectorMagnitude = bEnabled ? GetMagnitude() : 0.f;
-	ChannelData.Magnitude = bInvertType ? -EffectorMagnitude : EffectorMagnitude;
-}
-
-void ACEEffectorActor::OnTransformOptionsChanged()
-{
-	if (Mode != ECEClonerEffectorMode::Default)
-	{
-		return;
-	}
-
-	FVector ScaleDelta = GetScale();
-	ScaleDelta.X = FMath::Max(ScaleDelta.X, UE_KINDA_SMALL_NUMBER);
-	ScaleDelta.Y = FMath::Max(ScaleDelta.Y, UE_KINDA_SMALL_NUMBER);
-	ScaleDelta.Z = FMath::Max(ScaleDelta.Z, UE_KINDA_SMALL_NUMBER);
-
-	ChannelData.LocationDelta = GetOffset();
-	ChannelData.RotationDelta = GetRotation().Quaternion();
-	ChannelData.ScaleDelta = ScaleDelta;
-}
-
-void ACEEffectorActor::OnTargetActorChanged()
-{
-	AActor* TargetActor = TargetActorWeak.Get();
-	AActor* const InternalTargetActor = InternalTargetActorWeak.Get();
-	if (TargetActor && TargetActor == InternalTargetActor)
-	{
-		OnTargetOptionsChanged();
-		return;
-	}
+			RadialExtension->SetRadialAngle(RadialAngle);
+			RadialExtension->SetRadialMinRadius(RadialMinRadius);
+			RadialExtension->SetRadialMaxRadius(RadialMaxRadius);
+			RadialExtension->SetInvertType(bInvertType);
+			RadialExtension->SetEasing(Easing);
+		}
 
-	// unbind, except if it is self
-	if (InternalTargetActor && InternalTargetActor->GetRootComponent())
-	{
-		if (InternalTargetActor != this)
+		if (UCEEffectorTorusType* TorusExtension = EffectorComponent->FindOrAddExtension<UCEEffectorTorusType>())
 		{
-			InternalTargetActor->OnDestroyed.RemoveAll(this);
-			InternalTargetActor->GetRootComponent()->TransformUpdated.RemoveAll(this);
-			InternalTargetActorWeak.Reset();
+			TorusExtension->SetTorusRadius(TorusRadius);
+			TorusExtension->SetTorusInnerRadius(TorusInnerRadius);
+			TorusExtension->SetTorusOuterRadius(TorusOuterRadius);
+			TorusExtension->SetInvertType(bInvertType);
+			TorusExtension->SetEasing(Easing);
 		}
-	}
-
-	// set self by default if invalid
-	if (!TargetActor || !TargetActor->GetRootComponent())
-	{
-		TargetActor = this;
-	}
-
-	// bind to transform event, do not bind to self since we already do that
-	if (TargetActor != this)
-	{
-		TargetActor->GetRootComponent()->TransformUpdated.RemoveAll(this);
-		TargetActor->GetRootComponent()->TransformUpdated.AddUObject(this, &ACEEffectorActor::OnTargetActorTransformChanged);
-		TargetActor->OnDestroyed.RemoveAll(this);
-		TargetActor->OnDestroyed.AddUniqueDynamic(this, &ACEEffectorActor::OnTargetActorDestroyed);
-	}
-
-	TargetActorWeak = TargetActor;
-	InternalTargetActorWeak = TargetActor;
-	OnTargetOptionsChanged();
-}
-
-void ACEEffectorActor::OnTargetActorDestroyed(AActor* InActor)
-{
-	const AActor* TargetActor = TargetActorWeak.Get();
-	if (TargetActor == InActor)
-	{
-		TargetActorWeak = this;
-		OnTargetActorChanged();
-	}
-}
-
-void ACEEffectorActor::OnNoiseFieldOptionsChanged()
-{
-	if (Mode != ECEClonerEffectorMode::NoiseField)
-	{
-		return;
-	}
-
-	ChannelData.LocationDelta = LocationStrength;
-	ChannelData.RotationDelta = RotationStrength.Quaternion();
-	ChannelData.ScaleDelta = ScaleStrength;
-	ChannelData.Frequency = Frequency;
-	ChannelData.Pan = Pan;
-}
-
-void ACEEffectorActor::OnPushOptionsChanged()
-{
-	if (Mode != ECEClonerEffectorMode::Push)
-	{
-		return;
-	}
-
-	ChannelData.LocationDelta = PushStrength;
-	ChannelData.RotationDelta = FQuat::Identity;
-	ChannelData.ScaleDelta = FVector::OneVector;
-	ChannelData.Pan = FVector(0, static_cast<float>(PushDirection), 0);
-}
-
-void ACEEffectorActor::OnTargetActorTransformChanged(USceneComponent*, EUpdateTransformFlags, ETeleportType)
-{
-	OnTargetOptionsChanged();
-}
-
-void ACEEffectorActor::OnTargetOptionsChanged()
-{
-	if (Mode != ECEClonerEffectorMode::Target)
-	{
-		return;
-	}
-
-	const AActor* InternalTargetActor = InternalTargetActorWeak.Get();
-	if (!InternalTargetActor)
-	{
-		return;
-	}
-
-	ChannelData.LocationDelta = InternalTargetActor->GetActorLocation();
-	ChannelData.RotationDelta = FQuat::Identity;
-	ChannelData.ScaleDelta = FVector::OneVector;
-}
-
-#if WITH_EDITOR
-void ACEEffectorActor::OnVisualizerOptionsChanged()
-{
-	const UCEClonerEffectorSettings* ClonerEffectorSettings = GetDefault<UCEClonerEffectorSettings>();
-
-	if (!ClonerEffectorSettings)
-	{
-		return;
-	}
-
-	if (InnerVisualizerMaterial)
-	{
-		InnerVisualizerMaterial->SetVectorParameterValue(VisualizerColorName, ClonerEffectorSettings->VisualizerInnerColor);
-	}
-
-	if (OuterVisualizerMaterial)
-	{
-		OuterVisualizerMaterial->SetVectorParameterValue(VisualizerColorName, ClonerEffectorSettings->VisualizerOuterColor);
-	}
-
-	if (InnerVisualizerComponent)
-	{
-		InnerVisualizerComponent->SetVisibility(bVisualizerComponentVisible, false);
-	}
-
-	if (OuterVisualizerComponent)
-	{
-		OuterVisualizerComponent->SetVisibility(bVisualizerComponentVisible, false);
-	}
-
-	UE::ClonerEffector::SetBillboardComponentSprite(this, TEXT("/Script/Engine.Texture2D'/ClonerEffector/Textures/T_EffectorIcon.T_EffectorIcon'"));
-	UE::ClonerEffector::SetBillboardComponentVisibility(this, bVisualizerSpriteVisible);
-}
-
-void ACEEffectorActor::OnEffectorDeveloperSettingsChanged(UObject* InSettings, FPropertyChangedEvent& InEvent)
-{
-	OnVisualizerOptionsChanged();
-}
-
-void ACEEffectorActor::UpdateVisualizer(int32 InVisualizerId, TFunction<void(UDynamicMesh*)> InMeshFunction) const
-{
-	UDynamicMeshComponent* MeshComponent = InVisualizerId == InnerVisualizerId ? InnerVisualizerComponent : OuterVisualizerComponent;
-
-	if (!MeshComponent)
-	{
-		return;
-	}
 
-	UDynamicMesh* DynamicMesh = MeshComponent->GetDynamicMesh();
-
-	DynamicMesh->EditMesh([](FDynamicMesh3& InMesh)
-	{
-		InMesh.Clear();
-	});
-
-	if (bVisualizerComponentVisible)
-	{
-		InMeshFunction(DynamicMesh);
-	}
+		if (UCEEffectorOffsetMode* OffsetExtension = EffectorComponent->FindOrAddExtension<UCEEffectorOffsetMode>())
+		{
+			OffsetExtension->SetOffset(Offset);
+			OffsetExtension->SetScale(Scale);
+			OffsetExtension->SetRotation(Rotation);
+		}
 
-	// Apply material
-	UMaterialInstanceDynamic* VisualizerMaterial = InVisualizerId == InnerVisualizerId ? InnerVisualizerMaterial : OuterVisualizerMaterial;
-	MeshComponent->SetMaterial(0, VisualizerMaterial);
-}
-#endif
+		if (UCEEffectorTargetMode* TargetExtension = EffectorComponent->FindOrAddExtension<UCEEffectorTargetMode>())
+		{
+			TargetExtension->SetTargetActorWeak(TargetActorWeak);
+		}
 
-void ACEEffectorActor::OnColorChanged()
-{
-	ChannelData.Color = Color;
-}
+		if (UCEEffectorNoiseMode* NoiseExtension = EffectorComponent->FindOrAddExtension<UCEEffectorNoiseMode>())
+		{
+			NoiseExtension->SetLocationStrength(LocationStrength);
+			NoiseExtension->SetRotationStrength(RotationStrength);
+			NoiseExtension->SetScaleStrength(ScaleStrength);
+			NoiseExtension->SetFrequency(Frequency);
+			NoiseExtension->SetPan(Pan);
+		}
 
-void ACEEffectorActor::UpdateEffectorTypes()
-{
-	OnSphereChanged();
-	OnBoxChanged();
-	OnPlaneChanged();
-	OnRadialChanged();
-	OnTorusChanged();
-	OnUnboundChanged();
-}
+		if (UCEEffectorPushMode* PushExtension = EffectorComponent->FindOrAddExtension<UCEEffectorPushMode>())
+		{
+			PushExtension->SetPushDirection(PushDirection);
+			PushExtension->SetPushStrength(PushStrength);
+		}
 
-void ACEEffectorActor::OnForceOptionsChanged()
-{
-	if (bOrientationForceEnabled)
-	{
-		ChannelData.OrientationForceRate = OrientationForceRate;
-		ChannelData.OrientationForceMin = OrientationForceMin;
-		ChannelData.OrientationForceMax = OrientationForceMax;
-	}
-	else
-	{
-		ChannelData.OrientationForceRate = 0.f;
-		ChannelData.OrientationForceMin = FVector::ZeroVector;
-		ChannelData.OrientationForceMax = FVector::ZeroVector;
-	}
+		if (UCEEffectorForceExtension* ForceExtension = EffectorComponent->FindOrAddExtension<UCEEffectorForceExtension>())
+		{
+			ForceExtension->SetForcesEnabled(bAttractionForceEnabled || bGravityForceEnabled || bOrientationForceEnabled || bVortexForceEnabled || bCurlNoiseForceEnabled);
 
-	if (bVortexForceEnabled)
-	{
-		ChannelData.VortexForceAmount = VortexForceAmount;
-		ChannelData.VortexForceAxis = VortexForceAxis;
-	}
-	else
-	{
-		ChannelData.VortexForceAmount = 0.f;
-		ChannelData.VortexForceAxis = FVector::ZeroVector;
-	}
+			ForceExtension->SetAttractionForceEnabled(bAttractionForceEnabled);
+			ForceExtension->SetAttractionForceFalloff(AttractionForceFalloff);
+			ForceExtension->SetAttractionForceStrength(AttractionForceStrength);
 
-	if (bCurlNoiseForceEnabled)
-	{
-		ChannelData.CurlNoiseForceStrength = CurlNoiseForceStrength;
-		ChannelData.CurlNoiseForceFrequency = CurlNoiseForceFrequency;
-	}
-	else
-	{
-		ChannelData.CurlNoiseForceStrength = 0.f;
-		ChannelData.CurlNoiseForceFrequency = 0.f;
-	}
+			ForceExtension->SetGravityForceEnabled(bGravityForceEnabled);
+			ForceExtension->SetGravityForceAcceleration(GravityForceAcceleration);
 
-	if (bAttractionForceEnabled)
-	{
-		ChannelData.AttractionForceStrength = AttractionForceStrength;
-		ChannelData.AttractionForceFalloff = AttractionForceFalloff;
-	}
-	else
-	{
-		ChannelData.AttractionForceStrength = 0.f;
-		ChannelData.AttractionForceFalloff = 0.f;
-	}
+			ForceExtension->SetOrientationForceEnabled(bOrientationForceEnabled);
+			ForceExtension->SetOrientationForceMin(OrientationForceMin);
+			ForceExtension->SetOrientationForceMax(OrientationForceMax);
+			ForceExtension->SetOrientationForceRate(OrientationForceRate);
 
-	if (bGravityForceEnabled)
-	{
-		ChannelData.GravityForceAcceleration = GravityForceAcceleration;
-	}
-	else
-	{
-		ChannelData.GravityForceAcceleration = FVector::ZeroVector;
-	}
-}
+			ForceExtension->SetVortexForceEnabled(bVortexForceEnabled);
+			ForceExtension->SetVortexForceAmount(VortexForceAmount);
+			ForceExtension->SetVortexForceAxis(VortexForceAxis);
 
-void ACEEffectorActor::OnForceEnabledChanged()
-{
-	if (bEnabled)
-	{
-		// Refresh cloners to reset clones transform
-		for (const TWeakObjectPtr<ACEClonerActor>& ClonerWeak : ClonersWeak)
-		{
-			if (ACEClonerActor* Cloner = ClonerWeak.Get())
-			{
-				Cloner->RequestClonerUpdate();
-			}
+			ForceExtension->SetCurlNoiseForceEnabled(bCurlNoiseForceEnabled);
+			ForceExtension->SetCurlNoiseForceFrequency(CurlNoiseForceFrequency);
+			ForceExtension->SetCurlNoiseForceStrength(CurlNoiseForceStrength);
 		}
 	}
 
-	OnForceOptionsChanged();
+	PRAGMA_ENABLE_DEPRECATION_WARNINGS
+
+	MigrateToVersion = INDEX_NONE;
 }
diff --git a/Engine/Plugins/Experimental/ClonerEffector/Source/ClonerEffector/Private/Effector/Extensions/CEEffectorForceExtension.cpp b/Engine/Plugins/Experimental/ClonerEffector/Source/ClonerEffector/Private/Effector/Extensions/CEEffectorForceExtension.cpp
new file mode 100644
--- /dev/null
+++ b/Engine/Plugins/Experimental/ClonerEffector/Source/ClonerEffector/Private/Effector/Extensions/CEEffectorForceExtension.cpp
@@ -0,0 +1,279 @@
+// Copyright Epic Games, Inc. All Rights Reserved.
+
+#include "Effector/Extensions/CEEffectorForceExtension.h"
+
+#include "CEClonerEffectorShared.h"
+#include "Effector/CEEffectorComponent.h"
+
+void UCEEffectorForceExtension::SetForcesEnabled(bool bInForcesEnabled)
+{
+	if (bForcesEnabled == bInForcesEnabled)
+	{
+		return;
+	}
+
+	bForcesEnabled = bInForcesEnabled;
+	UpdateExtensionParameters(/** UpdateCloners */ true);
+}
+
+void UCEEffectorForceExtension::SetAttractionForceStrength(float InForceStrength)
+{
+	if (FMath::IsNearlyEqual(AttractionForceStrength, InForceStrength))
+	{
+		return;
+	}
+
+	AttractionForceStrength = InForceStrength;
+	UpdateExtensionParameters();
+}
+
+void UCEEffectorForceExtension::SetAttractionForceFalloff(float InForceFalloff)
+{
+	if (FMath::IsNearlyEqual(AttractionForceFalloff, InForceFalloff))
+	{
+		return;
+	}
+
+	AttractionForceFalloff = InForceFalloff;
+	UpdateExtensionParameters();
+}
+
+void UCEEffectorForceExtension::SetGravityForceEnabled(bool bInForceEnabled)
+{
+	if (bGravityForceEnabled == bInForceEnabled)
+	{
+		return;
+	}
+
+	bGravityForceEnabled = bInForceEnabled;
+	UpdateExtensionParameters(/** UpdateCloners */ true);
+}
+
+void UCEEffectorForceExtension::SetGravityForceAcceleration(const FVector& InAcceleration)
+{
+	if (GravityForceAcceleration.Equals(InAcceleration))
+	{
+		return;
+	}
+
+	GravityForceAcceleration = InAcceleration;
+	UpdateExtensionParameters();
+}
+
+void UCEEffectorForceExtension::SetOrientationForceEnabled(bool bInForceEnabled)
+{
+	if (bOrientationForceEnabled == bInForceEnabled)
+	{
+		return;
+	}
+
+	bOrientationForceEnabled = bInForceEnabled;
+	UpdateExtensionParameters(/** UpdateCloners */ true);
+}
+
+void UCEEffectorForceExtension::SetOrientationForceRate(float InForceOrientationRate)
+{
+	if (FMath::IsNearlyEqual(OrientationForceRate, InForceOrientationRate))
+	{
+		return;
+	}
+
+	OrientationForceRate = InForceOrientationRate;
+	UpdateExtensionParameters();
+}
+
+void UCEEffectorForceExtension::SetOrientationForceMin(const FVector& InForceOrientationMin)
+{
+	if (OrientationForceMin.Equals(InForceOrientationMin))
+	{
+		return;
+	}
+
+	OrientationForceMin = InForceOrientationMin;
+	UpdateExtensionParameters();
+}
+
+void UCEEffectorForceExtension::SetOrientationForceMax(const FVector& InForceOrientationMax)
+{
+	if (OrientationForceMax.Equals(InForceOrientationMax))
+	{
+		return;
+	}
+
+	OrientationForceMax = InForceOrientationMax;
+	UpdateExtensionParameters();
+}
+
+void UCEEffectorForceExtension::SetVortexForceEnabled(bool bInForceEnabled)
+{
+	if (bVortexForceEnabled == bInForceEnabled)
+	{
+		return;
+	}
+
+	bVortexForceEnabled = bInForceEnabled;
+	UpdateExtensionParameters(/** UpdateCloners */ true);
+}
+
+void UCEEffectorForceExtension::SetVortexForceAmount(float InForceVortexAmount)
+{
+	if (FMath::IsNearlyEqual(VortexForceAmount, InForceVortexAmount))
+	{
+		return;
+	}
+
+	VortexForceAmount = InForceVortexAmount;
+	UpdateExtensionParameters();
+}
+
+void UCEEffectorForceExtension::SetVortexForceAxis(const FVector& InForceVortexAxis)
+{
+	if (VortexForceAxis.Equals(InForceVortexAxis))
+	{
+		return;
+	}
+
+	VortexForceAxis = InForceVortexAxis;
+	UpdateExtensionParameters();
+}
+
+void UCEEffectorForceExtension::SetCurlNoiseForceEnabled(bool bInForceEnabled)
+{
+	if (bCurlNoiseForceEnabled == bInForceEnabled)
+	{
+		return;
+	}
+
+	bCurlNoiseForceEnabled = bInForceEnabled;
+	UpdateExtensionParameters(/** UpdateCloners */ true);
+}
+
+void UCEEffectorForceExtension::SetCurlNoiseForceStrength(float InForceCurlNoiseStrength)
+{
+	if (FMath::IsNearlyEqual(CurlNoiseForceStrength, InForceCurlNoiseStrength))
+	{
+		return;
+	}
+
+	CurlNoiseForceStrength = InForceCurlNoiseStrength;
+	UpdateExtensionParameters();
+}
+
+void UCEEffectorForceExtension::SetCurlNoiseForceFrequency(float InForceCurlNoiseFrequency)
+{
+	if (FMath::IsNearlyEqual(CurlNoiseForceFrequency, InForceCurlNoiseFrequency))
+	{
+		return;
+	}
+
+	CurlNoiseForceFrequency = InForceCurlNoiseFrequency;
+	UpdateExtensionParameters();
+}
+
+void UCEEffectorForceExtension::SetAttractionForceEnabled(bool bInForceEnabled)
+{
+	if (bAttractionForceEnabled == bInForceEnabled)
+	{
+		return;
+	}
+
+	bAttractionForceEnabled = bInForceEnabled;
+	UpdateExtensionParameters(/** UpdateCloners */ true);
+}
+
+void UCEEffectorForceExtension::OnExtensionParametersChanged(UCEEffectorComponent* InComponent)
+{
+	Super::OnExtensionParametersChanged(InComponent);
+
+	FCEClonerEffectorChannelData& ChannelData = InComponent->GetChannelData();
+
+	if (bForcesEnabled && bOrientationForceEnabled)
+	{
+		ChannelData.OrientationForceRate = OrientationForceRate;
+		ChannelData.OrientationForceMin = OrientationForceMin;
+		ChannelData.OrientationForceMax = OrientationForceMax;
+	}
+	else
+	{
+		ChannelData.OrientationForceRate = 0.f;
+		ChannelData.OrientationForceMin = FVector::ZeroVector;
+		ChannelData.OrientationForceMax = FVector::ZeroVector;
+	}
+
+	if (bForcesEnabled && bVortexForceEnabled)
+	{
+		ChannelData.VortexForceAmount = VortexForceAmount;
+		ChannelData.VortexForceAxis = VortexForceAxis;
+	}
+	else
+	{
+		ChannelData.VortexForceAmount = 0.f;
+		ChannelData.VortexForceAxis = FVector::ZeroVector;
+	}
+
+	if (bForcesEnabled && bCurlNoiseForceEnabled)
+	{
+		ChannelData.CurlNoiseForceStrength = CurlNoiseForceStrength;
+		ChannelData.CurlNoiseForceFrequency = CurlNoiseForceFrequency;
+	}
+	else
+	{
+		ChannelData.CurlNoiseForceStrength = 0.f;
+		ChannelData.CurlNoiseForceFrequency = 0.f;
+	}
+
+	if (bForcesEnabled && bAttractionForceEnabled)
+	{
+		ChannelData.AttractionForceStrength = AttractionForceStrength;
+		ChannelData.AttractionForceFalloff = AttractionForceFalloff;
+	}
+	else
+	{
+		ChannelData.AttractionForceStrength = 0.f;
+		ChannelData.AttractionForceFalloff = 0.f;
+	}
+
+	if (bForcesEnabled && bGravityForceEnabled)
+	{
+		ChannelData.GravityForceAcceleration = GravityForceAcceleration;
+	}
+	else
+	{
+		ChannelData.GravityForceAcceleration = FVector::ZeroVector;
+	}
+}
+
+void UCEEffectorForceExtension::OnForceOptionsChanged()
+{
+	UpdateExtensionParameters(/** UpdateCloners */ true);
+}
+
+#if WITH_EDITOR
+const TCEPropertyChangeDispatcher<UCEEffectorForceExtension> UCEEffectorForceExtension::PropertyChangeDispatcher =
+{
+	/** Force */
+	{ GET_MEMBER_NAME_CHECKED(UCEEffectorForceExtension, bForcesEnabled), &UCEEffectorForceExtension::OnForceOptionsChanged },
+	{ GET_MEMBER_NAME_CHECKED(UCEEffectorForceExtension, bOrientationForceEnabled), &UCEEffectorForceExtension::OnForceOptionsChanged },
+	{ GET_MEMBER_NAME_CHECKED(UCEEffectorForceExtension, OrientationForceRate), &UCEEffectorForceExtension::OnForceOptionsChanged },
+	{ GET_MEMBER_NAME_CHECKED(UCEEffectorForceExtension, OrientationForceMin), &UCEEffectorForceExtension::OnForceOptionsChanged },
+	{ GET_MEMBER_NAME_CHECKED(UCEEffectorForceExtension, OrientationForceMax), &UCEEffectorForceExtension::OnForceOptionsChanged },
+	{ GET_MEMBER_NAME_CHECKED(UCEEffectorForceExtension, bVortexForceEnabled), &UCEEffectorForceExtension::OnForceOptionsChanged },
+	{ GET_MEMBER_NAME_CHECKED(UCEEffectorForceExtension, VortexForceAmount), &UCEEffectorForceExtension::OnForceOptionsChanged },
+	{ GET_MEMBER_NAME_CHECKED(UCEEffectorForceExtension, VortexForceAxis), &UCEEffectorForceExtension::OnForceOptionsChanged },
+	{ GET_MEMBER_NAME_CHECKED(UCEEffectorForceExtension, bCurlNoiseForceEnabled), &UCEEffectorForceExtension::OnForceOptionsChanged },
+	{ GET_MEMBER_NAME_CHECKED(UCEEffectorForceExtension, CurlNoiseForceStrength), &UCEEffectorForceExtension::OnForceOptionsChanged },
+	{ GET_MEMBER_NAME_CHECKED(UCEEffectorForceExtension, CurlNoiseForceFrequency), &UCEEffectorForceExtension::OnForceOptionsChanged },
+	{ GET_MEMBER_NAME_CHECKED(UCEEffectorForceExtension, bAttractionForceEnabled), &UCEEffectorForceExtension::OnForceOptionsChanged },
+	{ GET_MEMBER_NAME_CHECKED(UCEEffectorForceExtension, AttractionForceStrength), &UCEEffectorForceExtension::OnForceOptionsChanged },
+	{ GET_MEMBER_NAME_CHECKED(UCEEffectorForceExtension, AttractionForceFalloff), &UCEEffectorForceExtension::OnForceOptionsChanged },
+	{ GET_MEMBER_NAME_CHECKED(UCEEffectorForceExtension, bGravityForceEnabled), &UCEEffectorForceExtension::OnForceOptionsChanged },
+	{ GET_MEMBER_NAME_CHECKED(UCEEffectorForceExtension, GravityForceAcceleration), &UCEEffectorForceExtension::OnForceOptionsChanged },
+};
+
+void UCEEffectorForceExtension::PostEditChangeProperty(FPropertyChangedEvent& InPropertyChangedEvent)
+{
+	Super::PostEditChangeProperty(InPropertyChangedEvent);
+
+	PropertyChangeDispatcher.OnPropertyChanged(this, InPropertyChangedEvent);
+}
+#endif
diff --git a/Engine/Plugins/Experimental/ClonerEffector/Source/ClonerEffector/Private/Effector/Modes/CEEffectorNoiseMode.cpp b/Engine/Plugins/Experimental/ClonerEffector/Source/ClonerEffector/Private/Effector/Modes/CEEffectorNoiseMode.cpp
new file mode 100644
--- /dev/null
+++ b/Engine/Plugins/Experimental/ClonerEffector/Source/ClonerEffector/Private/Effector/Modes/CEEffectorNoiseMode.cpp
@@ -0,0 +1,93 @@
+// Copyright Epic Games, Inc. All Rights Reserved.
+
+#include "Effector/Modes/CEEffectorNoiseMode.h"
+
+#include "CEClonerEffectorShared.h"
+#include "Effector/CEEffectorComponent.h"
+
+void UCEEffectorNoiseMode::SetLocationStrength(const FVector& InStrength)
+{
+	if (LocationStrength.Equals(InStrength))
+	{
+		return;
+	}
+
+	LocationStrength = InStrength;
+	UpdateExtensionParameters();
+}
+
+void UCEEffectorNoiseMode::SetRotationStrength(const FRotator& InStrength)
+{
+	if (RotationStrength.Equals(InStrength))
+	{
+		return;
+	}
+
+	RotationStrength = InStrength;
+	UpdateExtensionParameters();
+}
+
+void UCEEffectorNoiseMode::SetScaleStrength(const FVector& InStrength)
+{
+	if (ScaleStrength.Equals(InStrength))
+	{
+		return;
+	}
+
+	ScaleStrength = InStrength;
+	UpdateExtensionParameters();
+}
+
+void UCEEffectorNoiseMode::SetPan(const FVector& InPan)
+{
+	if (Pan.Equals(InPan))
+	{
+		return;
+	}
+
+	Pan = InPan;
+	UpdateExtensionParameters();
+}
+
+void UCEEffectorNoiseMode::SetFrequency(float InFrequency)
+{
+	InFrequency = FMath::Max(InFrequency, 0);
+
+	if (FMath::IsNearlyEqual(Frequency, InFrequency))
+	{
+		return;
+	}
+
+	Frequency = InFrequency;
+	UpdateExtensionParameters();
+}
+
+void UCEEffectorNoiseMode::OnExtensionParametersChanged(UCEEffectorComponent* InComponent)
+{
+	Super::OnExtensionParametersChanged(InComponent);
+
+	FCEClonerEffectorChannelData& ChannelData = InComponent->GetChannelData();
+	ChannelData.LocationDelta = LocationStrength;
+	ChannelData.RotationDelta = RotationStrength.Quaternion();
+	ChannelData.ScaleDelta = ScaleStrength;
+	ChannelData.Frequency = Frequency;
+	ChannelData.Pan = Pan;
+}
+
+#if WITH_EDITOR
+const TCEPropertyChangeDispatcher<UCEEffectorNoiseMode> UCEEffectorNoiseMode::PropertyChangeDispatcher =
+{
+	{ GET_MEMBER_NAME_CHECKED(UCEEffectorNoiseMode, LocationStrength), &UCEEffectorNoiseMode::OnExtensionPropertyChanged },
+	{ GET_MEMBER_NAME_CHECKED(UCEEffectorNoiseMode, RotationStrength), &UCEEffectorNoiseMode::OnExtensionPropertyChanged },
+	{ GET_MEMBER_NAME_CHECKED(UCEEffectorNoiseMode, ScaleStrength), &UCEEffectorNoiseMode::OnExtensionPropertyChanged },
+	{ GET_MEMBER_NAME_CHECKED(UCEEffectorNoiseMode, Pan), &UCEEffectorNoiseMode::OnExtensionPropertyChanged },
+	{ GET_MEMBER_NAME_CHECKED(UCEEffectorNoiseMode, Frequency), &UCEEffectorNoiseMode::OnExtensionPropertyChanged },
+};
+
+void UCEEffectorNoiseMode::PostEditChangeProperty(FPropertyChangedEvent& InPropertyChangedEvent)
+{
+	Super::PostEditChangeProperty(InPropertyChangedEvent);
+
+	PropertyChangeDispatcher.OnPropertyChanged(this, InPropertyChangedEvent);
+}
+#endif
diff --git a/Engine/Plugins/Experimental/ClonerEffector/Source/ClonerEffector/Private/Settings/CEClonerEffectorSettings.cpp b/Engine/Plugins/Experimental/ClonerEffector/Source/ClonerEffector/Private/Settings/CEClonerEffectorSettings.cpp
--- a/Engine/Plugins/Experimental/ClonerEffector/Source/ClonerEffector/Private/Settings/CEClonerEffectorSettings.cpp
+++ b/Engine/Plugins/Experimental/ClonerEffector/Source/ClonerEffector/Private/Settings/CEClonerEffectorSettings.cpp
@@ -6,6 +6,10 @@
 #include "Materials/MaterialInterface.h"
 #include "UObject/SoftObjectPath.h"
 
+#if WITH_EDITOR
+#include "HAL/IConsoleManager.h"
+#endif
+
 UCEClonerEffectorSettings::UCEClonerEffectorSettings()
 {
 	CategoryName = TEXT("Motion Design");
@@ -14,3 +18,100 @@ UCEClonerEffectorSettings::UCEClonerEffectorSettings()
 	DefaultStaticMesh = TSoftObjectPtr<UStaticMesh>(FSoftObjectPath(DefaultStaticMeshPath));
 	DefaultMaterial = TSoftObjectPtr<UMaterialInterface>(FSoftObjectPath(DefaultMaterialPath));
 }
+
+UStaticMesh* UCEClonerEffectorSettings::GetDefaultStaticMesh() const
+{
+	return DefaultStaticMesh.LoadSynchronous();
+}
+
+UMaterialInterface* UCEClonerEffectorSettings::GetDefaultMaterial() const
+{
+	return DefaultMaterial.LoadSynchronous();
+}
+
+void UCEClonerEffectorSettings::PostInitProperties()
+{
+	Super::PostInitProperties();
+
+#if WITH_EDITOR
+	CVarTSRShadingRejectionFlickeringPeriod = IConsoleManager::Get().FindConsoleVariable(TEXT("r.TSR.ShadingRejection.Flickering.Period"));
+
+	if (CVarTSRShadingRejectionFlickeringPeriod)
+	{
+		CVarTSRShadingRejectionFlickeringPeriod->OnChangedDelegate().RemoveAll(this);
+		CVarTSRShadingRejectionFlickeringPeriod->OnChangedDelegate().AddUObject(this, &UCEClonerEffectorSettings::OnTSRShadingRejectionFlickeringPeriodChanged);
+	}
+
+	OnTSRShadingRejectionFlickeringPeriodChanged(CVarTSRShadingRejectionFlickeringPeriod);
+#endif
+}
+
+#if WITH_EDITOR
+void UCEClonerEffectorSettings::PostEditChangeProperty(FPropertyChangedEvent& InPropertyChangedEvent)
+{
+	if (InPropertyChangedEvent.GetMemberPropertyName() == GET_MEMBER_NAME_CHECKED(UCEClonerEffectorSettings, bReduceMotionGhosting))
+	{
+		OnReduceMotionGhostingChanged();
+	}
+
+	Super::PostEditChangeProperty(InPropertyChangedEvent);
+}
+
+void UCEClonerEffectorSettings::EnableNoFlicker()
+{
+	if (IsNoFlickerEnabled())
+	{
+		return;
+	}
+
+	PreviousCVarValue = CVarTSRShadingRejectionFlickeringPeriod->GetInt();
+	CVarTSRShadingRejectionFlickeringPeriod->Set(NoFlicker);
+}
+
+void UCEClonerEffectorSettings::DisableNoFlicker()
+{
+	if (!IsNoFlickerEnabled())
+	{
+		return;
+	}
+
+	if (PreviousCVarValue.IsSet())
+	{
+		CVarTSRShadingRejectionFlickeringPeriod->Set(PreviousCVarValue.GetValue());
+	}
+	else
+	{
+		CVarTSRShadingRejectionFlickeringPeriod->Set(*CVarTSRShadingRejectionFlickeringPeriod->GetDefaultValue());
+	}
+}
+
+bool UCEClonerEffectorSettings::IsNoFlickerEnabled() const
+{
+	if (!CVarTSRShadingRejectionFlickeringPeriod)
+	{
+		return false;
+	}
+
+	return CVarTSRShadingRejectionFlickeringPeriod->GetInt() == NoFlicker;
+}
+
+void UCEClonerEffectorSettings::OnReduceMotionGhostingChanged()
+{
+	if (bReduceMotionGhosting)
+	{
+		EnableNoFlicker();
+	}
+	else
+	{
+		DisableNoFlicker();
+	}
+}
+
+void UCEClonerEffectorSettings::OnTSRShadingRejectionFlickeringPeriodChanged(IConsoleVariable* InCVar)
+{
+	if (InCVar == CVarTSRShadingRejectionFlickeringPeriod)
+	{
+		bReduceMotionGhosting = IsNoFlickerEnabled();
+	}
+}
+#endif
diff --git a/Engine/Plugins/Experimental/ClonerEffector/Source/ClonerEffector/Private/Settings/CEClonerEffectorSettings.h b/Engine/Plugins/Experimental/ClonerEffector/Source/ClonerEffector/Private/Settings/CEClonerEffectorSettings.h
--- a/Engine/Plugins/Experimental/ClonerEffector/Source/ClonerEffector/Private/Settings/CEClonerEffectorSettings.h
+++ b/Engine/Plugins/Experimental/ClonerEffector/Source/ClonerEffector/Private/Settings/CEClonerEffectorSettings.h
@@ -17,11 +17,46 @@ class UCEClonerEffectorSettings : public UDeveloperSettings
 	GENERATED_BODY()
 
 public:
+	static constexpr int32 NoFlicker = 1;
 	static constexpr TCHAR DefaultStaticMeshPath[] = TEXT("/Script/Engine.StaticMesh'/Engine/BasicShapes/Cube.Cube'");
 	static constexpr TCHAR DefaultMaterialPath[] = TEXT("/Script/Engine.Material'/ClonerEffector/Materials/DefaultClonerMaterial.DefaultClonerMaterial'");
 
+	FLinearColor GetVisualizerInnerColor() const
+	{
+		return VisualizerInnerColor;
+	}
+
+	FLinearColor GetVisualizerOuterColor() const
+	{
+		return VisualizerOuterColor;
+	}
+
+	bool GetSpawnDefaultActorAttached() const
+	{
+		return bSpawnDefaultActorAttached;
+	}
+
+	UStaticMesh* GetDefaultStaticMesh() const;
+
+	UMaterialInterface* GetDefaultMaterial() const;
+
+#if WITH_EDITOR
+	bool GetReduceMotionGhosting() const
+	{
+		return bReduceMotionGhosting;
+	}
+#endif
+
+protected:
 	UCEClonerEffectorSettings();
 
+	//~ Begin UObject
+	virtual void PostInitProperties() override;
+#if WITH_EDITOR
+	virtual void PostEditChangeProperty(FPropertyChangedEvent& InPropertyChangedEvent) override;
+#endif
+	//~ End UObject
+
 	/** Inner visualizer color for effectors */
 	UPROPERTY(Config, EditAnywhere, Category="Effector")
 	FLinearColor VisualizerInnerColor = FLinearColor(255, 0, 0, 0.1);
@@ -41,4 +76,26 @@ public:
 	/** Default material used when spawning default actor attached */
 	UPROPERTY(Config, EditAnywhere, Category="Cloner")
 	TSoftObjectPtr<UMaterialInterface> DefaultMaterial;
+
+#if WITH_EDITORONLY_DATA
+	/** Reduces the r.TSR.ShadingRejection.Flickering.Period from 3 (default) to 1 if enabled to avoid ghosting artifacts when moving */
+	UPROPERTY(EditInstanceOnly, BlueprintReadWrite, AdvancedDisplay, Category="Cloner", meta=(AllowPrivateAccess="true"))
+	bool bReduceMotionGhosting = false;
+#endif
+
+private:
+#if WITH_EDITOR
+	void EnableNoFlicker();
+	void DisableNoFlicker();
+	bool IsNoFlickerEnabled() const;
+
+	void OnReduceMotionGhostingChanged();
+	void OnTSRShadingRejectionFlickeringPeriodChanged(IConsoleVariable* InCVar);
+
+	/** Allows to reduce ghosting artifacts when moving cloner instances */
+	IConsoleVariable* CVarTSRShadingRejectionFlickeringPeriod = nullptr;
+
+	/** Previous value to restore it when disabled */
+	TOptional<int32> PreviousCVarValue;
+#endif
 };
diff --git a/Engine/Plugins/Experimental/ClonerEffector/Source/ClonerEffector/Private/Subsystems/CEClonerSubsystem.cpp b/Engine/Plugins/Experimental/ClonerEffector/Source/ClonerEffector/Private/Subsystems/CEClonerSubsystem.cpp
--- a/Engine/Plugins/Experimental/ClonerEffector/Source/ClonerEffector/Private/Subsystems/CEClonerSubsystem.cpp
+++ b/Engine/Plugins/Experimental/ClonerEffector/Source/ClonerEffector/Private/Subsystems/CEClonerSubsystem.cpp
@@ -2,7 +2,8 @@
 
 #include "Subsystems/CEClonerSubsystem.h"
 
-#include "Cloner/CEClonerActor.h"
+#include "Cloner/CEClonerComponent.h"
+#include "Cloner/Extensions/CEClonerExtensionBase.h"
 #include "Cloner/Layouts/CEClonerCircleLayout.h"
 #include "Cloner/Layouts/CEClonerCylinderLayout.h"
 #include "Cloner/Layouts/CEClonerGridLayout.h"
@@ -17,12 +18,6 @@
 #include "UObject/Class.h"
 #include "UObject/UObjectIterator.h"
 
-#if WITH_EDITOR
-#include "HAL/IConsoleManager.h"
-
-UCEClonerSubsystem::FOnCVarChanged UCEClonerSubsystem::OnCVarChangedDelegate;
-#endif
-
 UCEClonerSubsystem::FOnSubsystemInitialized UCEClonerSubsystem::OnSubsystemInitializedDelegate;
 
 UCEClonerSubsystem* UCEClonerSubsystem::Get()
@@ -53,31 +48,10 @@ void UCEClonerSubsystem::Initialize(FSubsystemCollectionBase& Collection)
 	// Scan for new layouts
 	ScanForRegistrableClasses();
 
-#if WITH_EDITOR
-	CVarTSRShadingRejectionFlickeringPeriod = IConsoleManager::Get().FindConsoleVariable(TEXT("r.TSR.ShadingRejection.Flickering.Period"));
-
-	if (CVarTSRShadingRejectionFlickeringPeriod)
-	{
-		CVarTSRShadingRejectionFlickeringPeriod->OnChangedDelegate().AddUObject(this, &UCEClonerSubsystem::OnTSRShadingRejectionFlickeringPeriodChanged);
-	}
-#endif
-
 	OnSubsystemInitializedDelegate.Broadcast();
 }
 
-void UCEClonerSubsystem::Deinitialize()
-{
-	Super::Deinitialize();
-
-#if WITH_EDITOR
-	if (CVarTSRShadingRejectionFlickeringPeriod)
-	{
-		CVarTSRShadingRejectionFlickeringPeriod->OnChangedDelegate().RemoveAll(this);
-	}
-#endif
-}
-
-bool UCEClonerSubsystem::RegisterLayoutClass(const UClass* InClonerLayoutClass)
+bool UCEClonerSubsystem::RegisterLayoutClass(UClass* InClonerLayoutClass)
 {
 	if (!IsValid(InClonerLayoutClass))
 	{
@@ -110,7 +84,7 @@ bool UCEClonerSubsystem::RegisterLayoutClass(const UClass* InClonerLayoutClass)
 
 	// Does not overwrite existing layouts
 	const FName LayoutName = CDO->GetLayoutName();
-	if (LayoutClasses.Contains(LayoutName))
+	if (LayoutName.IsNone() || LayoutClasses.Contains(LayoutName))
 	{
 		return false;
 	}
@@ -120,38 +94,34 @@ bool UCEClonerSubsystem::RegisterLayoutClass(const UClass* InClonerLayoutClass)
 	return true;
 }
 
-bool UCEClonerSubsystem::UnregisterLayoutClass(const UClass* InClonerLayoutClass)
+bool UCEClonerSubsystem::UnregisterLayoutClass(UClass* InClonerLayoutClass)
 {
 	if (!IsValid(InClonerLayoutClass))
 	{
 		return false;
 	}
 
-	for (TMap<FName, TSubclassOf<UCEClonerLayoutBase>>::TIterator It(LayoutClasses); It; ++It)
+	TSubclassOf<UCEClonerLayoutBase> LayoutClass(InClonerLayoutClass);
+	if (const FName* LayoutName = LayoutClasses.FindKey(LayoutClass))
 	{
-		if (It->Value.Get() == InClonerLayoutClass)
-		{
-			It.RemoveCurrent();
-			return true;
-		}
+		LayoutClasses.Remove(*LayoutName);
+		return true;
 	}
 
 	return false;
 }
 
-bool UCEClonerSubsystem::IsLayoutClassRegistered(const UClass* InClonerLayoutClass)
+bool UCEClonerSubsystem::IsLayoutClassRegistered(UClass* InClonerLayoutClass)
 {
 	if (!IsValid(InClonerLayoutClass))
 	{
 		return false;
 	}
 
-	for (const TPair<FName, TSubclassOf<UCEClonerLayoutBase>>& LayoutClassPair : LayoutClasses)
+	TSubclassOf<UCEClonerLayoutBase> LayoutClass(InClonerLayoutClass);
+	if (const FName* LayoutName = LayoutClasses.FindKey(LayoutClass))
 	{
-		if (LayoutClassPair.Value.Get() == InClonerLayoutClass)
-		{
-			return true;
-		}
+		return true;
 	}
 
 	return false;
@@ -172,92 +142,157 @@ UCEClonerSubsystem::FOnGetOrderedActors& UCEClonerSubsystem::GetCustomActorResol
 	return ActorResolver;
 }
 
-TArray<FName> UCEClonerSubsystem::GetLayoutNames() const
+bool UCEClonerSubsystem::RegisterExtensionClass(UClass* InClass)
 {
-	TArray<FName> LayoutNames;
-	LayoutClasses.GenerateKeyArray(LayoutNames);
-	return LayoutNames;
-}
+	if (!IsValid(InClass))
+	{
+		return false;
+	}
 
-FName UCEClonerSubsystem::FindLayoutName(TSubclassOf<UCEClonerLayoutBase> InLayoutClass) const
-{
-	if (const FName* Key = LayoutClasses.FindKey(InLayoutClass))
+	if (!InClass->IsChildOf(UCEClonerExtensionBase::StaticClass())
+		|| InClass->HasAnyClassFlags(CLASS_Abstract | CLASS_Deprecated | CLASS_NewerVersionExists))
 	{
-		return *Key;
+		return false;
 	}
 
-	return NAME_None;
+	if (IsExtensionClassRegistered(InClass))
+	{
+		return false;
+	}
+
+	const UCEClonerExtensionBase* CDO = InClass->GetDefaultObject<UCEClonerExtensionBase>();
+
+	if (!CDO)
+	{
+		return false;
+	}
+
+	const FName ExtensionName = CDO->GetExtensionName();
+	if (ExtensionName.IsNone() || ExtensionClasses.Contains(ExtensionName))
+	{
+		return false;
+	}
+
+	ExtensionClasses.Add(ExtensionName, CDO->GetClass());
+
+	return true;
 }
 
-UCEClonerLayoutBase* UCEClonerSubsystem::CreateNewLayout(FName InLayoutName, ACEClonerActor* InClonerActor)
+bool UCEClonerSubsystem::UnregisterExtensionClass(UClass* InClass)
 {
-	if (!IsValid(InClonerActor))
+	if (!IsValid(InClass))
 	{
-		return nullptr;
+		return false;
 	}
 
-	TSubclassOf<UCEClonerLayoutBase> const* LayoutClass = LayoutClasses.Find(InLayoutName);
-
-	if (!LayoutClass)
+	TSubclassOf<UCEClonerExtensionBase> ExtensionClass(InClass);
+	if (const FName* ExtensionName = ExtensionClasses.FindKey(ExtensionClass))
 	{
-		return nullptr;
+		ExtensionClasses.Remove(*ExtensionName);
+		return true;
 	}
 
-	return NewObject<UCEClonerLayoutBase>(InClonerActor, LayoutClass->Get());
+	return false;
 }
 
-void UCEClonerSubsystem::ScanForRegistrableClasses()
+bool UCEClonerSubsystem::IsExtensionClassRegistered(UClass* InClass) const
 {
-	for (const UClass* const Class : TObjectRange<UClass>())
+	if (!IsValid(InClass))
 	{
-		RegisterLayoutClass(Class);
+		return false;
 	}
+
+	TSubclassOf<UCEClonerExtensionBase> ExtensionClass(InClass);
+	return !!ExtensionClasses.FindKey(ExtensionClass);
 }
 
-#if WITH_EDITOR
-void UCEClonerSubsystem::EnableNoFlicker()
+TArray<FName> UCEClonerSubsystem::GetExtensionNames() const
 {
-	if (IsNoFlickerEnabled())
+	TArray<FName> ExtensionNames;
+	ExtensionClasses.GenerateKeyArray(ExtensionNames);
+	return ExtensionNames;
+}
+
+FName UCEClonerSubsystem::FindExtensionName(TSubclassOf<UCEClonerExtensionBase> InClass) const
+{
+	if (const FName* Key = ExtensionClasses.FindKey(InClass))
 	{
-		return;
+		return *Key;
 	}
 
-	PreviousCVarValue = CVarTSRShadingRejectionFlickeringPeriod->GetInt();
-	CVarTSRShadingRejectionFlickeringPeriod->Set(NoFlicker);
+	return NAME_None;
 }
 
-void UCEClonerSubsystem::DisableNoFlicker()
+UCEClonerExtensionBase* UCEClonerSubsystem::CreateNewExtension(FName InExtensionName, UCEClonerComponent* InCloner)
 {
-	if (!IsNoFlickerEnabled())
+	if (!IsValid(InCloner))
 	{
-		return;
+		return nullptr;
 	}
 
-	if (PreviousCVarValue.IsSet())
+	TSubclassOf<UCEClonerExtensionBase> const* ExtensionClass = ExtensionClasses.Find(InExtensionName);
+
+	if (!ExtensionClass)
 	{
-		CVarTSRShadingRejectionFlickeringPeriod->Set(PreviousCVarValue.GetValue());
+		return nullptr;
 	}
-	else
+
+	return NewObject<UCEClonerExtensionBase>(InCloner, ExtensionClass->Get(), NAME_None, RF_Transactional);
+}
+
+TArray<FName> UCEClonerSubsystem::GetLayoutNames() const
+{
+	TArray<FName> LayoutNames;
+	LayoutClasses.GenerateKeyArray(LayoutNames);
+	return LayoutNames;
+}
+
+FName UCEClonerSubsystem::FindLayoutName(TSubclassOf<UCEClonerLayoutBase> InLayoutClass) const
+{
+	if (const FName* Key = LayoutClasses.FindKey(InLayoutClass))
 	{
-		CVarTSRShadingRejectionFlickeringPeriod->Set(*CVarTSRShadingRejectionFlickeringPeriod->GetDefaultValue());
+		return *Key;
 	}
+
+	return NAME_None;
 }
 
-bool UCEClonerSubsystem::IsNoFlickerEnabled() const
+UCEClonerLayoutBase* UCEClonerSubsystem::CreateNewLayout(FName InLayoutName, UCEClonerComponent* InCloner)
 {
-	if (!CVarTSRShadingRejectionFlickeringPeriod)
+	if (!IsValid(InCloner))
 	{
-		return false;
+		return nullptr;
 	}
 
-	return CVarTSRShadingRejectionFlickeringPeriod->GetInt() == NoFlicker;
+	TSubclassOf<UCEClonerLayoutBase> const* LayoutClass = LayoutClasses.Find(InLayoutName);
+
+	if (!LayoutClass)
+	{
+		return nullptr;
+	}
+
+	return NewObject<UCEClonerLayoutBase>(InCloner, LayoutClass->Get());
 }
 
-void UCEClonerSubsystem::OnTSRShadingRejectionFlickeringPeriodChanged(IConsoleVariable* InCVar) const
+void UCEClonerSubsystem::ScanForRegistrableClasses()
 {
-	if (InCVar == CVarTSRShadingRejectionFlickeringPeriod)
 	{
-		OnCVarChangedDelegate.Broadcast();
+		TArray<UClass*> DerivedLayoutClasses;
+		GetDerivedClasses(UCEClonerLayoutBase::StaticClass(), DerivedLayoutClasses, true);
+
+		for (UClass* LayoutClass : DerivedLayoutClasses)
+		{
+			RegisterLayoutClass(LayoutClass);
+		}
 	}
-}
-#endif
\ No newline at end of file
+
+	{
+		TArray<UClass*> DerivedExtensionClasses;
+		GetDerivedClasses(UCEClonerExtensionBase::StaticClass(), DerivedExtensionClasses, true);
+
+		for (UClass* ExtensionClass : DerivedExtensionClasses)
+		{
+			RegisterExtensionClass(ExtensionClass);
+		}
+	}
+}
\ No newline at end of file
diff --git a/Engine/Plugins/Experimental/ClonerEffector/Source/ClonerEffector/Private/Subsystems/CEEffectorSubsystem.cpp b/Engine/Plugins/Experimental/ClonerEffector/Source/ClonerEffector/Private/Subsystems/CEEffectorSubsystem.cpp
--- a/Engine/Plugins/Experimental/ClonerEffector/Source/ClonerEffector/Private/Subsystems/CEEffectorSubsystem.cpp
+++ b/Engine/Plugins/Experimental/ClonerEffector/Source/ClonerEffector/Private/Subsystems/CEEffectorSubsystem.cpp
@@ -2,21 +2,33 @@
 
 #include "Subsystems/CEEffectorSubsystem.h"
 
-#include "Effector/CEEffectorActor.h"
+#include "Engine/Engine.h"
 #include "NiagaraDataChannel.h"
 #include "NiagaraDataChannelAccessor.h"
 #include "NiagaraDataChannelPublic.h"
+#include "Effector/CEEffectorComponent.h"
+#include "Effector/CEEffectorExtensionBase.h"
+#include "Effector/Modes/CEEffectorNoiseMode.h"
+#include "Effector/Modes/CEEffectorOffsetMode.h"
+#include "Effector/Modes/CEEffectorPushMode.h"
+#include "Effector/Modes/CEEffectorTargetMode.h"
+#include "Effector/Types/CEEffectorBoxType.h"
+#include "Effector/Types/CEEffectorPlaneType.h"
+#include "Effector/Types/CEEffectorRadialType.h"
+#include "Effector/Types/CEEffectorSphereType.h"
+#include "Effector/Types/CEEffectorTorusType.h"
+#include "Effector/Types/CEEffectorUnboundType.h"
 
 DEFINE_LOG_CATEGORY_STATIC(LogACEEffectorSubsystem, Log, All);
 
 UCEEffectorSubsystem::FOnSubsystemInitialized UCEEffectorSubsystem::OnSubsystemInitializedDelegate;
 UCEEffectorSubsystem::FOnEffectorIdentifierChanged UCEEffectorSubsystem::OnEffectorIdentifierChangedDelegate;
 
-UCEEffectorSubsystem* UCEEffectorSubsystem::Get(const UWorld* InWorld)
+UCEEffectorSubsystem* UCEEffectorSubsystem::Get()
 {
-	if (InWorld)
+	if (GEngine)
 	{
-		return InWorld->GetSubsystem<UCEEffectorSubsystem>();
+		return GEngine->GetEngineSubsystem<UCEEffectorSubsystem>();
 	}
 
 	return nullptr;
@@ -30,21 +42,29 @@ void UCEEffectorSubsystem::Initialize(FSubsystemCollectionBase& Collection)
 	EffectorDataChannelAsset = LoadObject<UNiagaraDataChannelAsset>(nullptr, DataChannelAssetPath);
 
 	check(EffectorDataChannelAsset->Get());
-}
 
-void UCEEffectorSubsystem::PostInitialize()
-{
-	Super::PostInitialize();
+	// Register types
+	RegisterExtensionClass(UCEEffectorSphereType::StaticClass());
+	RegisterExtensionClass(UCEEffectorPlaneType::StaticClass());
+	RegisterExtensionClass(UCEEffectorBoxType::StaticClass());
+	RegisterExtensionClass(UCEEffectorUnboundType::StaticClass());
+	RegisterExtensionClass(UCEEffectorRadialType::StaticClass());
+	RegisterExtensionClass(UCEEffectorTorusType::StaticClass());
 
-	if (const UWorld* World = GetWorld())
-	{
-		OnSubsystemInitializedDelegate.Broadcast(World);
-	}
+	// Register modes
+	RegisterExtensionClass(UCEEffectorOffsetMode::StaticClass());
+	RegisterExtensionClass(UCEEffectorTargetMode::StaticClass());
+	RegisterExtensionClass(UCEEffectorNoiseMode::StaticClass());
+	RegisterExtensionClass(UCEEffectorPushMode::StaticClass());
+
+	ScanForRegistrableClasses();
+
+	OnSubsystemInitializedDelegate.Broadcast();
 }
 
-bool UCEEffectorSubsystem::RegisterChannelEffector(ACEEffectorActor* InEffector)
+bool UCEEffectorSubsystem::RegisterChannelEffector(UCEEffectorComponent* InEffector)
 {
-	if (!IsValid(InEffector))
+	if (!IsValid(InEffector) || !InEffector->GetOwner())
 	{
 		return false;
 	}
@@ -54,22 +74,22 @@ bool UCEEffectorSubsystem::RegisterChannelEffector(ACEEffectorActor* InEffector)
 	if (EffectorIndex == INDEX_NONE)
 	{
 		EffectorIndex = EffectorsWeak.Add(InEffector);
-		UE_LOG(LogACEEffectorSubsystem, Log, TEXT("%s effector registered in channel %i"), *InEffector->GetActorNameOrLabel(), EffectorIndex);
+		UE_LOG(LogACEEffectorSubsystem, Log, TEXT("%s effector registered in channel %i"), *InEffector->GetOwner()->GetActorNameOrLabel(), EffectorIndex);
 	}
 
-	if (InEffector->ChannelData.Identifier != EffectorIndex)
+	if (InEffector->GetChannelData().Identifier != EffectorIndex)
 	{
-		const int32 OldIdentifier = InEffector->ChannelData.Identifier;
-		InEffector->ChannelData.Identifier = EffectorIndex;
+		const int32 OldIdentifier = InEffector->GetChannelData().Identifier;
+		InEffector->GetChannelData().Identifier = EffectorIndex;
 		OnEffectorIdentifierChangedDelegate.Broadcast(InEffector, OldIdentifier, EffectorIndex);
 	}
 
 	return true;
 }
 
-bool UCEEffectorSubsystem::UnregisterChannelEffector(ACEEffectorActor* InEffector)
+bool UCEEffectorSubsystem::UnregisterChannelEffector(UCEEffectorComponent* InEffector)
 {
-	if (!InEffector)
+	if (!InEffector || !InEffector->GetOwner())
 	{
 		return false;
 	}
@@ -78,21 +98,21 @@ bool UCEEffectorSubsystem::UnregisterChannelEffector(ACEEffectorActor* InEffecto
 
 	if (bUnregistered)
 	{
-		UE_LOG(LogACEEffectorSubsystem, Log, TEXT("%s effector unregistered from channel"), *InEffector->GetActorNameOrLabel());
+		UE_LOG(LogACEEffectorSubsystem, Log, TEXT("%s effector unregistered from channel"), *InEffector->GetOwner()->GetActorNameOrLabel());
 
-		const int32 OldIdentifier = InEffector->ChannelData.Identifier;
-		InEffector->ChannelData.Identifier = INDEX_NONE;
-		OnEffectorIdentifierChangedDelegate.Broadcast(InEffector, OldIdentifier, InEffector->ChannelData.Identifier);
+		const int32 OldIdentifier = InEffector->GetChannelData().Identifier;
+		InEffector->GetChannelData().Identifier = INDEX_NONE;
+		OnEffectorIdentifierChangedDelegate.Broadcast(InEffector, OldIdentifier, InEffector->GetChannelData().Identifier);
 	}
 
 	return bUnregistered;
 }
 
-ACEEffectorActor* UCEEffectorSubsystem::GetEffectorByChannelIdentifier(int32 InIdentifier) const
+UCEEffectorComponent* UCEEffectorSubsystem::GetEffectorByChannelIdentifier(int32 InIdentifier) const
 {
 	if (EffectorsWeak.IsValidIndex(InIdentifier))
 	{
-		if (ACEEffectorActor* Effector = EffectorsWeak[InIdentifier].Get())
+		if (UCEEffectorComponent* Effector = EffectorsWeak[InIdentifier].Get())
 		{
 			if (Effector->GetChannelIdentifier() == InIdentifier)
 			{
@@ -104,11 +124,121 @@ ACEEffectorActor* UCEEffectorSubsystem::GetEffectorByChannelIdentifier(int32 InI
 	return nullptr;
 }
 
-void UCEEffectorSubsystem::UpdateEffectorChannel()
+bool UCEEffectorSubsystem::RegisterExtensionClass(UClass* InClass)
+{
+	if (!IsValid(InClass))
+	{
+		return false;
+	}
+
+	if (!InClass->IsChildOf(UCEEffectorExtensionBase::StaticClass())
+		|| InClass->HasAnyClassFlags(CLASS_Abstract | CLASS_Deprecated | CLASS_NewerVersionExists))
+	{
+		return false;
+	}
+
+	if (IsExtensionClassRegistered(InClass))
+	{
+		return false;
+	}
+
+	const UCEEffectorExtensionBase* CDO = InClass->GetDefaultObject<UCEEffectorExtensionBase>();
+
+	if (!CDO)
+	{
+		return false;
+	}
+
+	const FName ExtensionName = CDO->GetExtensionName();
+
+	if (ExtensionName.IsNone() || ExtensionClasses.Contains(ExtensionName))
+	{
+		return false;
+	}
+
+	ExtensionClasses.Add(ExtensionName, CDO->GetClass());
+
+	return true;
+}
+
+bool UCEEffectorSubsystem::UnregisterExtensionClass(UClass* InClass)
+{
+	if (!IsValid(InClass))
+	{
+		return false;
+	}
+
+	TSubclassOf<UCEEffectorExtensionBase> ExtensionClass(InClass);
+
+	if (const FName* ExtensionName = ExtensionClasses.FindKey(ExtensionClass))
+	{
+		ExtensionClasses.Remove(*ExtensionName);
+		return true;
+	}
+
+	return false;
+}
+
+bool UCEEffectorSubsystem::IsExtensionClassRegistered(UClass* InClass) const
+{
+	TSubclassOf<UCEEffectorExtensionBase> ExtensionClass(InClass);
+	return !!ExtensionClasses.FindKey(ExtensionClass);
+}
+
+TArray<FName> UCEEffectorSubsystem::GetExtensionNames(TSubclassOf<UCEEffectorExtensionBase> InExtensionClass) const
+{
+	TArray<FName> ExtensionNames;
+
+	const FName ExtensionName = FindExtensionName(InExtensionClass);
+
+	if (!ExtensionName.IsNone())
+	{
+		ExtensionNames.Add(ExtensionName);
+	}
+	else
+	{
+		for (const TPair<FName, TSubclassOf<UCEEffectorExtensionBase>>& ExtensionPair : ExtensionClasses)
+		{
+			if (ExtensionPair.Value && ExtensionPair.Value->IsChildOf(InExtensionClass))
+			{
+				ExtensionNames.Add(ExtensionPair.Key);
+			}
+		}
+	}
+
+	return ExtensionNames;
+}
+
+FName UCEEffectorSubsystem::FindExtensionName(TSubclassOf<UCEEffectorExtensionBase> InClass) const
 {
-	const UWorld* World = GetWorld();
+	if (const FName* ExtensionName = ExtensionClasses.FindKey(InClass))
+	{
+		return *ExtensionName;
+	}
+
+	return NAME_None;
+}
+
+UCEEffectorExtensionBase* UCEEffectorSubsystem::CreateNewExtension(FName InExtensionName, UCEEffectorComponent* InEffector)
+{
+	if (!IsValid(InEffector))
+	{
+		return nullptr;
+	}
 
-	if (!IsValid(World))
+	TSubclassOf<UCEEffectorExtensionBase> const* ExtensionClass = ExtensionClasses.Find(InExtensionName);
+
+	if (!ExtensionClass)
+	{
+		return nullptr;
+	}
+
+	return NewObject<UCEEffectorExtensionBase>(InEffector, ExtensionClass->Get(), NAME_None, RF_Transactional);
+}
+
+void UCEEffectorSubsystem::UpdateEffectorChannel(const UWorld* InWorld)
+{
+	if (!IsValid(InWorld))
 	{
 		return;
 	}
@@ -120,7 +250,7 @@ void UCEEffectorSubsystem::UpdateEffectorChannel()
 
 	// Reserve space in channel for each effectors
 	static const FNiagaraDataChannelSearchParameters SearchParameters;
-	UNiagaraDataChannelWriter* ChannelWriter = UNiagaraDataChannelLibrary::WriteToNiagaraDataChannel(World, EffectorDataChannelAsset.Get(), SearchParameters, EffectorsWeak.Num(), true, true, true, UCEEffectorSubsystem::StaticClass()->GetName());
+	UNiagaraDataChannelWriter* ChannelWriter = UNiagaraDataChannelLibrary::WriteToNiagaraDataChannel(InWorld, EffectorDataChannelAsset.Get(), SearchParameters, EffectorsWeak.Num(), true, true, true, UCEEffectorSubsystem::StaticClass()->GetName());
 
 	if (!ChannelWriter)
 	{
@@ -130,9 +260,9 @@ void UCEEffectorSubsystem::UpdateEffectorChannel()
 
 	// Remove invalid effectors and push updates to effector assigned channel indexes
 	int32 EffectorIndex = 0;
-	for (TArray<TWeakObjectPtr<ACEEffectorActor>>::TIterator It(EffectorsWeak); It; ++It)
+	for (TArray<TWeakObjectPtr<UCEEffectorComponent>>::TIterator It(EffectorsWeak); It; ++It)
 	{
-		ACEEffectorActor* Effector = It->Get();
+		UCEEffectorComponent* Effector = It->Get();
 
 		if (!IsValid(Effector))
 		{
@@ -140,6 +270,11 @@ void UCEEffectorSubsystem::UpdateEffectorChannel()
 			continue;
 		}
 
+		if (Effector->GetWorld() != InWorld)
+		{
+			continue;
+		}
+
 		FCEClonerEffectorChannelData& ChannelData = Effector->GetChannelData();
 
 		const int32 PreviousIdentifier = ChannelData.Identifier;
@@ -159,9 +294,15 @@ void UCEEffectorSubsystem::UpdateEffectorChannel()
 	}
 }
 
-bool UCEEffectorSubsystem::IsTickableInEditor() const
+void UCEEffectorSubsystem::ScanForRegistrableClasses()
 {
-	return true;
+	TArray<UClass*> DerivedExtensionClasses;
+	GetDerivedClasses(UCEEffectorExtensionBase::StaticClass(), DerivedExtensionClasses, true);
+
+	for (UClass* ExtensionClass : DerivedExtensionClasses)
+	{
+		RegisterExtensionClass(ExtensionClass);
+	}
 }
 
 TStatId UCEEffectorSubsystem::GetStatId() const
@@ -171,5 +312,25 @@ TStatId UCEEffectorSubsystem::GetStatId() const
 
 void UCEEffectorSubsystem::Tick(float InDeltaTime)
 {
-	UpdateEffectorChannel();
+	TSet<UWorld*> Worlds;
+
+	for (const TWeakObjectPtr<UCEEffectorComponent>& EffectorWeak : EffectorsWeak)
+	{
+		UCEEffectorComponent* Effector = EffectorWeak.Get();
+
+		if (Effector && Effector->GetWorld())
+		{
+			Worlds.Add(Effector->GetWorld());
+		}
+	}
+
+	for (UWorld* World : Worlds)
+	{
+		UpdateEffectorChannel(World);
+	}
+}
+
+bool UCEEffectorSubsystem::IsTickable() const
+{
+	return !EffectorsWeak.IsEmpty();
 }
diff --git a/Engine/Plugins/Experimental/ClonerEffector/Source/ClonerEffector/Private/Subsystems/CEEffectorSubsystem.h b/Engine/Plugins/Experimental/ClonerEffector/Source/ClonerEffector/Private/Subsystems/CEEffectorSubsystem.h
--- a/Engine/Plugins/Experimental/ClonerEffector/Source/ClonerEffector/Private/Subsystems/CEEffectorSubsystem.h
+++ b/Engine/Plugins/Experimental/ClonerEffector/Source/ClonerEffector/Private/Subsystems/CEEffectorSubsystem.h
@@ -3,57 +3,98 @@
 #pragma once
 
 #include "Engine/World.h"
-#include "Subsystems/WorldSubsystem.h"
+#include "Subsystems/EngineSubsystem.h"
 #include "UObject/WeakObjectPtrTemplates.h"
 #include "CEEffectorSubsystem.generated.h"
 
-class ACEEffectorActor;
+class UCEEffectorComponent;
+class UCEEffectorExtensionBase;
 class UNiagaraDataChannelAsset;
 
 UCLASS()
-class UCEEffectorSubsystem : public UTickableWorldSubsystem
+class UCEEffectorSubsystem : public UEngineSubsystem, public FTickableGameObject
 {
 	GENERATED_BODY()
 
 public:
-	DECLARE_MULTICAST_DELEGATE_OneParam(FOnSubsystemInitialized, const UWorld*)
-	static FOnSubsystemInitialized OnSubsystemInitializedDelegate;
-
-	/** Broadcast when this effector identifier changed to update linked cloners */
-	DECLARE_MULTICAST_DELEGATE_ThreeParams(FOnEffectorIdentifierChanged, ACEEffectorActor* /** InEffectorActor */, int32 /** OldIdentifier */, int32 /** NewIdentifier */)
-	static FOnEffectorIdentifierChanged OnEffectorIdentifierChangedDelegate;
+	static TMulticastDelegateRegistration<void()>& OnSubsystemInitialized()
+	{
+		return OnSubsystemInitializedDelegate;
+	}
 
-	static inline constexpr TCHAR DataChannelAssetPath[] = TEXT("/Script/Niagara.NiagaraDataChannelAsset'/ClonerEffector/Channels/NDC_Effector.NDC_Effector'");
+	static TMulticastDelegateRegistration<void(UCEEffectorComponent*, int32, int32)>& OnEffectorIdentifierChanged()
+	{
+		return OnEffectorIdentifierChangedDelegate;
+	}
 
 	/** Get this subsystem instance */
-	static UCEEffectorSubsystem* Get(const UWorld* InWorld = GWorld);
-
-	//~ Begin USubsystem
-	virtual void Initialize(FSubsystemCollectionBase& Collection) override;
-	virtual void PostInitialize() override;
-	//~ End USubsystem
+	static UCEEffectorSubsystem* Get();
 
 	/** Registers an effector actor to use it within a effector channel */
-	bool RegisterChannelEffector(ACEEffectorActor* InEffector);
+	bool RegisterChannelEffector(UCEEffectorComponent* InEffector);
 
 	/** Unregister an effector actor used within a effector channel */
-	bool UnregisterChannelEffector(ACEEffectorActor* InEffector);
+	bool UnregisterChannelEffector(UCEEffectorComponent* InEffector);
+
+	/** Get the effector using this channel identifier */
+	UCEEffectorComponent* GetEffectorByChannelIdentifier(int32 InIdentifier) const;
+
+	bool RegisterExtensionClass(UClass* InClass);
+
+	bool UnregisterExtensionClass(UClass* InClass);
+
+	bool IsExtensionClassRegistered(UClass* InClass) const;
+
+	template<typename InExtensionClass>
+	TArray<FName> GetExtensionNames() const
+	{
+		return GetExtensionNames(InExtensionClass::StaticClass());
+	}
+
+	TArray<FName> GetExtensionNames(TSubclassOf<UCEEffectorExtensionBase> InExtensionClass) const;
 
-	ACEEffectorActor* GetEffectorByChannelIdentifier(int32 InIdentifier) const;
+	/** Based on a extension class, find extension name */
+	FName FindExtensionName(TSubclassOf<UCEEffectorExtensionBase> InClass) const;
+
+	/** Creates a new extension instance for an effector */
+	UCEEffectorExtensionBase* CreateNewExtension(FName InExtensionName, UCEEffectorComponent* InEffector);
 
 protected:
-	//~ Begin FTickableGameObject interface
-	virtual bool IsTickableInEditor() const override;
+	static constexpr TCHAR DataChannelAssetPath[] = TEXT("/Script/Niagara.NiagaraDataChannelAsset'/ClonerEffector/Channels/NDC_Effector.NDC_Effector'");
+
+	/** Broadcasted when this subsystem is initialized */
+	DECLARE_MULTICAST_DELEGATE(FOnSubsystemInitialized)
+	static FOnSubsystemInitialized OnSubsystemInitializedDelegate;
+
+	/** Broadcasted when this effector identifier changed to update linked cloners */
+	DECLARE_MULTICAST_DELEGATE_ThreeParams(FOnEffectorIdentifierChanged, UCEEffectorComponent* /** InEffector */, int32 /** OldIdentifier */, int32 /** NewIdentifier */)
+	static FOnEffectorIdentifierChanged OnEffectorIdentifierChangedDelegate;
+
+	//~ Begin USubsystem
+	virtual void Initialize(FSubsystemCollectionBase& Collection) override;
+	//~ End USubsystem
+
+	//~ Begin FTickableGameObject
 	virtual TStatId GetStatId() const override;
-	virtual void Tick(float InDeltaTime) override;
-	//~ End FTickableGameObject interface
+	virtual void Tick(float InDeltaSeconds) override;
+	virtual bool IsTickable() const override;
+	virtual bool IsTickableInEditor() const override { return true; }
+	virtual ETickableTickType GetTickableTickType() const override { return ETickableTickType::Conditional; }
+	//~ End FTickableGameObject
 
 	/** Updates all registered effectors */
-	void UpdateEffectorChannel();
+	void UpdateEffectorChannel(const UWorld* InWorld);
+
+	/** Scan classes and registers them */
+	void ScanForRegistrableClasses();
+
+	/** Linking name to the extension class */
+	UPROPERTY()
+	TMap<FName, TSubclassOf<UCEEffectorExtensionBase>> ExtensionClasses;
 
 	/** Ordered effectors included in this channel */
 	UPROPERTY()
-	TArray<TWeakObjectPtr<ACEEffectorActor>> EffectorsWeak;
+	TArray<TWeakObjectPtr<UCEEffectorComponent>> EffectorsWeak;
 
 	/** This represents the data channel structure for effector */
 	UPROPERTY()
diff --git a/Engine/Plugins/Experimental/ClonerEffector/Source/ClonerEffector/Public/CEClonerEffectorShared.h b/Engine/Plugins/Experimental/ClonerEffector/Source/ClonerEffector/Public/CEClonerEffectorShared.h
--- a/Engine/Plugins/Experimental/ClonerEffector/Source/ClonerEffector/Public/CEClonerEffectorShared.h
+++ b/Engine/Plugins/Experimental/ClonerEffector/Source/ClonerEffector/Public/CEClonerEffectorShared.h
@@ -324,11 +324,15 @@ struct FCEClonerAttachmentTree
 	UPROPERTY()
 	ECEClonerAttachmentStatus Status = ECEClonerAttachmentStatus::Updated;
 
+	/** Attachment items that are dirty and need an update */
+	TSet<FCEClonerAttachmentItem*> DirtyItemAttachments;
+
 	void Reset()
 	{
 		ItemAttachmentMap.Empty();
 		RootActors.Empty();
 		MergedBakedMeshes.Empty();
+		DirtyItemAttachments.Empty();
 		Status = ECEClonerAttachmentStatus::Updated;
 	}
 };
