MotionDesign : ClonerEffector

diff --git a/Engine/Plugins/Experimental/ClonerEffector/Source/ClonerEffector/Public/Cloner/CEClonerActor.h b/Engine/Plugins/Experimental/ClonerEffector/Source/ClonerEffector/Public/Cloner/CEClonerActor.h
--- a/Engine/Plugins/Experimental/ClonerEffector/Source/ClonerEffector/Public/Cloner/CEClonerActor.h
+++ b/Engine/Plugins/Experimental/ClonerEffector/Source/ClonerEffector/Public/Cloner/CEClonerActor.h
@@ -7,7 +7,6 @@
 #include "GameFramework/Actor.h"
 #include "Layouts/CEClonerLayoutBase.h"
 #include "NiagaraMeshRendererProperties.h"
-#include "CEPropertyChangeDispatcher.h"
 #include "CEClonerActor.generated.h"
 
 class AActor;
@@ -17,519 +16,28 @@ class UCEClonerLayoutBase;
 class UNiagaraDataInterfaceCurve;
 class USceneComponent;
 
-UCLASS(MinimalAPI, BlueprintType, HideCategories=(Rendering,Replication,Collision,HLOD,Physics,Networking,Input,Actor,Cooking,LevelInstance,Streaming,DataLayers), DisplayName = "Motion Design Cloner Actor")
+UCLASS(MinimalAPI
+	, BlueprintType
+	, AutoExpandCategories=(ClonerComponent)
+	, HideCategories=(Rendering,Replication,Collision,HLOD,Physics,Networking,Input,Actor,Cooking,LevelInstance,Streaming,DataLayers)
+	, DisplayName = "Motion Design Cloner Actor")
 class ACEClonerActor : public AActor
 {
 	GENERATED_BODY()
 
-	friend class ACEEffectorActor;
-	friend class UCEClonerLayoutBase;
-
-	friend class FAvaClonerActorVisualizer;
-	friend class FCEEditorClonerDetailCustomization;
-
 public:
 	static inline const FString DefaultLabel = TEXT("Cloner");
 
 	ACEClonerActor();
 
-	UFUNCTION(BlueprintCallable, Category="Cloner")
-	CLONEREFFECTOR_API void SetTreeUpdateInterval(float InInterval);
-
-	UFUNCTION(BlueprintPure, Category="Cloner")
-	float GetTreeUpdateInterval() const
-	{
-		return TreeUpdateInterval;
-	}
-
-	UFUNCTION(BlueprintCallable, Category="Cloner")
-	CLONEREFFECTOR_API void SetMeshRenderMode(ECEClonerMeshRenderMode InMode);
-
-	UFUNCTION(BlueprintPure, Category="Cloner")
-	ECEClonerMeshRenderMode GetMeshRenderMode() const
-	{
-		return MeshRenderMode;
-	}
-
-	UFUNCTION(BlueprintCallable, Category="Cloner")
-	CLONEREFFECTOR_API void SetMeshFacingMode(ENiagaraMeshFacingMode InMode);
-
-	UFUNCTION(BlueprintPure, Category="Cloner")
-	ENiagaraMeshFacingMode GetMeshFacingMode() const
-	{
-		return MeshFacingMode;
-	}
-
-	UFUNCTION(BlueprintCallable, Category="Cloner")
-	CLONEREFFECTOR_API void SetMeshCastShadows(bool InbCastShadows);
-
-	UFUNCTION(BlueprintPure, Category="Cloner")
-	bool GetMeshCastShadows() const
-	{
-		return bMeshCastShadows;
-	}
-
-	CLONEREFFECTOR_API void SetDefaultMeshes(const TArray<TObjectPtr<UStaticMesh>>& InMeshes);
-
-	const TArray<TObjectPtr<UStaticMesh>>& GetDefaultMeshes() const
-	{
-		return DefaultMeshes;
-	}
-
-	UFUNCTION(BlueprintCallable, Category="Cloner", meta=(DisplayName="SetDefaultMeshes"))
-	CLONEREFFECTOR_API void BP_SetDefaultMeshes(const TArray<UStaticMesh*>& InMeshes);
-
-	UFUNCTION(BlueprintPure, Category="Cloner", meta=(DisplayName="GetDefaultMeshes"))
-	CLONEREFFECTOR_API TArray<UStaticMesh*> BP_GetDefaultMeshes() const;
-
-	UFUNCTION(BlueprintCallable, Category="Cloner")
-	CLONEREFFECTOR_API void SetUseOverrideMaterial(bool bInOverride);
-
-	UFUNCTION(BlueprintPure, Category="Cloner")
-	bool GetUseOverrideMaterial() const
-	{
-		return bUseOverrideMaterial;
-	}
-
-	UFUNCTION(BlueprintCallable, Category="Cloner")
-	CLONEREFFECTOR_API void SetOverrideMaterial(UMaterialInterface* InMaterial);
-
-	UFUNCTION(BlueprintPure, Category="Cloner")
-	UMaterialInterface* GetOverrideMaterial() const
-	{
-		return OverrideMaterial;
-	}
-
-	UFUNCTION(BlueprintCallable, Category="Cloner")
-	CLONEREFFECTOR_API void SetSeed(int32 InSeed);
-
-	UFUNCTION(BlueprintPure, Category="Cloner")
-	int32 GetSeed() const
-	{
-		return Seed;
-	}
-
-	UFUNCTION(BlueprintCallable, Category="Cloner")
-	CLONEREFFECTOR_API void SetColor(const FLinearColor& InColor);
-
-	UFUNCTION(BlueprintPure, Category="Cloner")
-	const FLinearColor& GetColor() const
-	{
-		return Color;
-	}
-
-	UFUNCTION(BlueprintCallable, Category="Cloner")
-	CLONEREFFECTOR_API void SetVisualizeEffectors(bool bInVisualize);
-
-	UFUNCTION(BlueprintPure, Category="Cloner")
-	bool GetVisualizeEffectors() const
-	{
-		return bVisualizeEffectors;
-	}
-
-	UFUNCTION(BlueprintCallable, Category="Cloner")
-	CLONEREFFECTOR_API void SetSurfaceCollisionEnabled(bool bInSurfaceCollisionEnabled);
-
-	UFUNCTION(BlueprintPure, Category="Cloner")
-	bool GetSurfaceCollisionEnabled() const
-	{
-		return bSurfaceCollisionEnabled;
-	}
-
-	UFUNCTION(BlueprintCallable, Category="Cloner")
-	CLONEREFFECTOR_API void SetParticleCollisionEnabled(bool bInParticleCollisionEnabled);
-
-	UFUNCTION(BlueprintPure, Category="Cloner")
-	bool GetParticleCollisionEnabled() const
-	{
-		return bParticleCollisionEnabled;
-	}
-
-	UFUNCTION(BlueprintCallable, Category="Cloner")
-	CLONEREFFECTOR_API void SetCollisionVelocityEnabled(bool bInCollisionVelocityEnabled);
-
-	UFUNCTION(BlueprintPure, Category="Cloner")
-	bool GetCollisionVelocityEnabled() const
-	{
-		return bCollisionVelocityEnabled;
-	}
-
-	UFUNCTION(BlueprintCallable, Category="Cloner")
-	CLONEREFFECTOR_API void SetCollisionIterations(int32 InCollisionIterations);
-
-	UFUNCTION(BlueprintPure, Category="Cloner")
-	int32 GetCollisionIterations() const
-	{
-		return CollisionIterations;
-	}
-
-	UFUNCTION(BlueprintCallable, Category="Cloner")
-	CLONEREFFECTOR_API void SetCollisionGridResolution(int32 InCollisionGridResolution);
-
-	UFUNCTION(BlueprintPure, Category="Cloner")
-	int32 GetCollisionGridResolution() const
-	{
-		return CollisionGridResolution;
-	}
-
-	UFUNCTION(BlueprintCallable, Category="Cloner")
-	CLONEREFFECTOR_API void SetCollisionGridSize(const FVector& InCollisionGridSize);
-
-	UFUNCTION(BlueprintPure, Category="Cloner")
-	FVector GetCollisionGridSize() const
-	{
-		return CollisionGridSize;
-	}
-
-	UFUNCTION(BlueprintCallable, Category="Cloner")
-	CLONEREFFECTOR_API void SetCollisionRadiusMode(ECEClonerCollisionRadiusMode InMode);
-
-	UFUNCTION(BlueprintPure, Category="Cloner")
-	ECEClonerCollisionRadiusMode GetCollisionRadiusMode() const
-	{
-		return CollisionRadiusMode;
-	}
-
-	UFUNCTION(BlueprintPure, Category="Cloner")
-	const TArray<float>& GetCollisionRadii() const
-	{
-		return CollisionRadii;
-	}
-
-	UFUNCTION(BlueprintCallable, Category="Cloner")
-	CLONEREFFECTOR_API void SetMassMin(float InMassMin);
-
-	UFUNCTION(BlueprintPure, Category="Cloner")
-	float GetMassMin() const
-	{
-		return MassMin;
-	}
-
-	UFUNCTION(BlueprintCallable, Category="Cloner")
-	CLONEREFFECTOR_API void SetMassMax(float InMassMax);
-
-	UFUNCTION(BlueprintPure, Category="Cloner")
-	float GetMassMax() const
-	{
-		return MassMax;
-	}
-
-	UFUNCTION(BlueprintCallable, Category="Cloner")
-	CLONEREFFECTOR_API void SetDeltaStepEnabled(bool bInEnabled);
-
-	UFUNCTION(BlueprintPure, Category="Cloner")
-	bool GetDeltaStepEnabled() const
-	{
-		return bDeltaStepEnabled;
-	}
-
-	UFUNCTION(BlueprintCallable, Category="Cloner")
-	CLONEREFFECTOR_API void SetDeltaStepRotation(const FRotator& InRotation);
-
-	UFUNCTION(BlueprintPure, Category="Cloner")
-	FRotator GetDeltaStepRotation() const
-	{
-		return DeltaStepRotation;
-	}
-
-	UFUNCTION(BlueprintCallable, Category="Cloner")
-	CLONEREFFECTOR_API void SetDeltaStepScale(const FVector& InScale);
-
-	UFUNCTION(BlueprintPure, Category="Cloner")
-	FVector GetDeltaStepScale() const
-	{
-		return DeltaStepScale;
-	}
-
-	UFUNCTION(BlueprintCallable, Category="Cloner")
-	CLONEREFFECTOR_API void SetEnabled(bool bInEnable);
-
-	UFUNCTION(BlueprintPure, Category="Cloner")
-	bool GetEnabled() const
-	{
-		return bEnabled;
-	}
-
-	UFUNCTION(BlueprintCallable, Category="Cloner")
-	CLONEREFFECTOR_API void SetInvertProgress(bool bInInvertProgress);
-
-	UFUNCTION(BlueprintPure, Category="Cloner")
-	bool GetInvertProgress() const
-	{
-		return bInvertProgress;
-	}
-
-	UFUNCTION(BlueprintCallable, Category="Cloner")
-	CLONEREFFECTOR_API void SetProgress(float InProgress);
-
-	UFUNCTION(BlueprintPure, Category="Cloner")
-	float GetProgress() const
-	{
-		return Progress;
-	}
-
-	UFUNCTION(BlueprintCallable, Category="Cloner")
-	CLONEREFFECTOR_API void SetLayoutName(FName InLayoutName);
-
-	UFUNCTION(BlueprintPure, Category="Cloner")
-	FName GetLayoutName() const
-	{
-		return LayoutName;
-	}
-
-	UFUNCTION(BlueprintCallable, Category="Cloner")
-	CLONEREFFECTOR_API void SetRangeEnabled(bool bInRangeEnabled);
-
-	UFUNCTION(BlueprintPure, Category="Cloner")
-	bool GetRangeEnabled() const
-	{
-		return bRangeEnabled;
-	}
-
-	UFUNCTION(BlueprintCallable, Category="Cloner")
-	CLONEREFFECTOR_API void SetRangeOffsetMin(const FVector& InRangeOffsetMin);
-
-	UFUNCTION(BlueprintPure, Category="Cloner")
-	const FVector& GetRangeOffsetMin() const
-	{
-		return RangeOffsetMin;
-	}
-
-	UFUNCTION(BlueprintCallable, Category="Cloner")
-	CLONEREFFECTOR_API void SetRangeOffsetMax(const FVector& InRangeOffsetMax);
-
-	UFUNCTION(BlueprintPure, Category="Cloner")
-	const FVector& GetRangeOffsetMax() const
-	{
-		return RangeOffsetMax;
-	}
-
-	UFUNCTION(BlueprintCallable, Category="Cloner")
-	CLONEREFFECTOR_API void SetRangeRotationMin(const FRotator& InRangeRotationMin);
-
-	UFUNCTION(BlueprintPure, Category="Cloner")
-	const FRotator& GetRangeRotationMin() const
-	{
-		return RangeRotationMin;
-	}
-
-	UFUNCTION(BlueprintCallable, Category="Cloner")
-	CLONEREFFECTOR_API void SetRangeRotationMax(const FRotator& InRangeRotationMax);
-
-	UFUNCTION(BlueprintPure, Category="Cloner")
-	const FRotator& GetRangeRotationMax() const
-	{
-		return RangeRotationMax;
-	}
-
-	UFUNCTION(BlueprintCallable, Category="Cloner")
-	CLONEREFFECTOR_API void SetRangeScaleUniform(bool bInRangeScaleUniform);
-
-	UFUNCTION(BlueprintPure, Category="Cloner")
-	bool GetRangeScaleUniform() const
-	{
-		return bRangeScaleUniform;
-	}
-
-	UFUNCTION(BlueprintCallable, Category="Cloner")
-	CLONEREFFECTOR_API void SetRangeScaleMin(const FVector& InRangeScaleMin);
-
-	UFUNCTION(BlueprintPure, Category="Cloner")
-	const FVector& GetRangeScaleMin() const
-	{
-		return RangeScaleMin;
-	}
-
-	UFUNCTION(BlueprintCallable, Category="Cloner")
-	CLONEREFFECTOR_API void SetRangeScaleMax(const FVector& InRangeScaleMax);
-
-	UFUNCTION(BlueprintPure, Category="Cloner")
-	const FVector& GetRangeScaleMax() const
-	{
-		return RangeScaleMax;
-	}
-
-	UFUNCTION(BlueprintCallable, Category="Cloner")
-	CLONEREFFECTOR_API void SetRangeScaleUniformMin(float InRangeScaleUniformMin);
-
-	UFUNCTION(BlueprintPure, Category="Cloner")
-	float GetRangeScaleUniformMin() const
-	{
-		return RangeScaleUniformMin;
-	}
-
-	UFUNCTION(BlueprintCallable, Category="Cloner")
-	CLONEREFFECTOR_API void SetRangeScaleUniformMax(float InRangeScaleUniformMax);
-
-	UFUNCTION(BlueprintPure, Category="Cloner")
-	float GetRangeScaleUniformMax() const
-	{
-		return RangeScaleUniformMax;
-	}
-
-	UFUNCTION(BlueprintCallable, Category="Cloner")
-	CLONEREFFECTOR_API void SetSpawnLoopMode(ECEClonerSpawnLoopMode InMode);
-
-	UFUNCTION(BlueprintPure, Category="Cloner")
-	ECEClonerSpawnLoopMode GetSpawnLoopMode() const
-	{
-		return SpawnLoopMode;
-	}
-
-	UFUNCTION(BlueprintCallable, Category="Cloner")
-	CLONEREFFECTOR_API void SetSpawnLoopIterations(int32 InIterations);
-
-	UFUNCTION(BlueprintPure, Category="Cloner")
-	int32 GetSpawnLoopIterations() const
-	{
-		return SpawnLoopIterations;
-	}
-
-	UFUNCTION(BlueprintCallable, Category="Cloner")
-	CLONEREFFECTOR_API void SetSpawnLoopInterval(float InInterval);
-
-	UFUNCTION(BlueprintPure, Category="Cloner")
-	float GetSpawnLoopInterval() const
-	{
-		return SpawnLoopInterval;
-	}
-
-	UFUNCTION(BlueprintCallable, Category="Cloner")
-	CLONEREFFECTOR_API void SetSpawnBehaviorMode(ECEClonerSpawnBehaviorMode InMode);
-
-	UFUNCTION(BlueprintPure, Category="Cloner")
-	ECEClonerSpawnBehaviorMode GetSpawnBehaviorMode() const
-	{
-		return SpawnBehaviorMode;
-	}
-
-	UFUNCTION(BlueprintCallable, Category="Cloner")
-	CLONEREFFECTOR_API void SetSpawnRate(float InRate);
-
-	UFUNCTION(BlueprintPure, Category="Cloner")
-	float GetSpawnRate() const
-	{
-		return SpawnRate;
-	}
-
-	UFUNCTION(BlueprintCallable, Category="Cloner")
-	CLONEREFFECTOR_API void SetLifetimeEnabled(bool bInEnabled);
-
-	UFUNCTION(BlueprintPure, Category="Cloner")
-	bool GetLifetimeEnabled() const
-	{
-		return bLifetimeEnabled;
-	}
-
-	UFUNCTION(BlueprintCallable, Category="Cloner")
-	CLONEREFFECTOR_API void SetLifetimeMin(float InMin);
-
-	UFUNCTION(BlueprintPure, Category="Cloner")
-	float GetLifetimeMin() const
-	{
-		return LifetimeMin;
-	}
-
-	UFUNCTION(BlueprintCallable, Category="Cloner")
-	CLONEREFFECTOR_API void SetLifetimeMax(float InMax);
-
-	UFUNCTION(BlueprintPure, Category="Cloner")
-	float GetLifetimeMax() const
-	{
-		return LifetimeMax;
-	}
-
-	UFUNCTION(BlueprintCallable, Category="Cloner")
-	CLONEREFFECTOR_API void SetLifetimeScaleEnabled(bool bInEnabled);
-
-	UFUNCTION(BlueprintPure, Category="Cloner")
-	bool GetLifetimeScaleEnabled() const
-	{
-		return bLifetimeScaleEnabled;
-	}
-
-	UFUNCTION()
-	CLONEREFFECTOR_API void SetLifetimeScaleCurve(const FRichCurve& InCurve);
-
-	UFUNCTION()
-	const FRichCurve& GetLifetimeScaleCurve() const
-	{
-		return LifetimeScaleCurve;
-	}
-
-	UFUNCTION(BlueprintPure, Category="Cloner")
-	CLONEREFFECTOR_API UCEClonerLayoutBase* GetActiveLayout() const;
-
-	template<
-		typename InLayoutClass
-		UE_REQUIRES(TIsDerivedFrom<InLayoutClass, UCEClonerLayoutBase>::Value)>
-	bool IsActiveLayout() const
-	{
-		if (const UCEClonerLayoutBase* CurrentLayout = GetActiveLayout())
-		{
-			return CurrentLayout->GetClass() == InLayoutClass::StaticClass();
-		}
-
-		return false;
-	}
-
-	template<
-		typename InLayoutClass
-		UE_REQUIRES(TIsDerivedFrom<InLayoutClass, UCEClonerLayoutBase>::Value)>
-	InLayoutClass* GetActiveLayout() const
-	{
-		return Cast<InLayoutClass>(GetActiveLayout());
-	}
-
-	/** Returns the number of meshes this cloner currently handles */
-	UFUNCTION(BlueprintPure, Category="Cloner")
-	CLONEREFFECTOR_API int32 GetMeshCount() const;
-
-	/** Gets the number of effectors applied on this cloner */
-	UFUNCTION(BlueprintPure, Category="Cloner")
-	CLONEREFFECTOR_API int32 GetEffectorCount() const;
-
-#if WITH_EDITOR
-	/** This will force an update of the cloner attachment tree */
-	UFUNCTION(CallInEditor, Category="Cloner")
-	CLONEREFFECTOR_API void ForceUpdateCloner();
-
-	/** This will spawn an effector actor and link it to this cloner */
-	UFUNCTION(CallInEditor, Category="Cloner")
-	CLONEREFFECTOR_API void SpawnLinkedEffector();
-
-	/** When the cloner actor is created, we spawn a default actor attached */
-	void SpawnDefaultActorAttached();
-#endif
-
-	/** Links a new effector to apply transformation on clones */
-	CLONEREFFECTOR_API bool LinkEffector(ACEEffectorActor* InEffector);
-
-	/** Unlinks the effector and reset the cloner simulation */
-	CLONEREFFECTOR_API bool UnlinkEffector(ACEEffectorActor* InEffector);
-
-	/** Checks if an effector is linked with this cloner */
-	CLONEREFFECTOR_API bool IsEffectorLinked(const ACEEffectorActor* InEffector) const;
-
 protected:
 	//~ Begin UObject
-	virtual void PostInitProperties() override;
+	virtual void Serialize(FArchive& InArchive) override;
 	virtual void PostLoad() override;
-#if WITH_EDITOR
-	virtual void PostEditUndo() override;
-	virtual void PostEditChangeProperty(FPropertyChangedEvent& InPropertyChangedEvent) override;
-#endif
 	//~ End UObject
 
 	//~ Begin AActor
 	virtual void PostActorCreated() override;
-	/** Allow ticking in editor viewport too */
-	virtual bool ShouldTickIfViewportsOnly() const override
-	{
-		return true;
-	}
-	/** Tick needed to keep the attachment tree updated */
-	virtual void Tick(float DeltaSeconds) override;
 #if WITH_EDITOR
 	virtual FString GetDefaultActorLabel() const override;
 #endif
@@ -540,308 +48,224 @@ protected:
 		return ClonerComponent;
 	}
 
-	UNiagaraDataInterfaceCurve* GetLifetimeScaleCurveDI() const
-	{
-		return LifetimeScaleCurveDIWeak.Get();
-	}
-
-	/** Will force a system update to refresh user parameters */
-	void RequestClonerUpdate(bool bInImmediate = false);
-
-	void UpdateLayoutOptions();
-
-	/** Used by effector actors to apply transformations to this cloner instances */
-	const FCEClonerEffectorDataInterfaces* GetEffectorDataInterfaces() const;
-
-	/** For each valid effector that has this cloner registered */
-	void ForEachEffector(TFunctionRef<bool(ACEEffectorActor*, int32)> InFunction);
-
-	void OnClonerMeshUpdated(UCEClonerComponent* InClonerComponent);
-	void OnClonerSystemChanged();
+private:
+	void OnClonerInitialized(UCEClonerComponent* InClonerComponent);
 
-	void OnEffectorIdentifierChanged(ACEEffectorActor* InEffector, int32 InOldIdentifier, int32 InNewIdentifier);
+	void MigrateDeprecatedProperties();
 
-	void OnEffectorsChanged();
-	void OnEnabledChanged();
-	void OnMeshRenderModeChanged();
-	void OnMeshRendererOptionsChanged();
-	void OnDefaultMeshesChanged();
-	void OnOverrideMaterialChanged();
-	void OnSeedChanged();
-	void OnColorChanged();
-	void OnProgressChanged();
-	void OnDeltaStepChanged();
-	void OnRangeOptionsChanged();
-	void OnLayoutNameChanged();
-	void OnSpawnOptionsChanged();
-	void OnLifetimeScaleCurveChanged();
-	void OnLifetimeOptionsChanged();
-	void OnCollisionOptionsChanged();
+#if WITH_EDITOR
+	/** Hide outline selection when this cloner only is selected in viewport */
+	void OnEditorSelectionChanged(UObject* InSelection);
 
-	/** Update sprite visibility of this cloner */
-	void OnVisualizerSpriteVisibleChanged();
+	TOptional<bool> UseSelectionOutline;
+#endif
 
-	bool IsClonerValid() const;
+	int32 MigrateToVersion = INDEX_NONE;
 
-#if WITH_EDITOR
-	void OnReduceMotionGhostingChanged();
-	void OnCVarChanged();
-#endif
+	UPROPERTY(VisibleInstanceOnly, Category="Cloner")
+	TObjectPtr<UCEClonerComponent> ClonerComponent;
 
-	/** Is this cloner enabled/disabled */
-	UPROPERTY(EditAnywhere, BlueprintReadWrite, Setter="SetEnabled", Getter="GetEnabled", Category="Cloner")
+	UE_DEPRECATED(5.5, "Moved to component and extensions")
+	UPROPERTY()
 	bool bEnabled = true;
 
-	/** Interval to update the attachment tree and update the cloner meshes */
-	UPROPERTY(EditAnywhere, BlueprintReadWrite, Setter="SetTreeUpdateInterval", Getter="GetTreeUpdateInterval", Category="Cloner", AdvancedDisplay, meta=(ClampMin="0.0"))
+	UE_DEPRECATED(5.5, "Moved to component and extensions")
+	UPROPERTY()
 	float TreeUpdateInterval = 0.2f;
 
-	/** Cloner instance seed for random deterministic patterns */
-	UPROPERTY(EditAnywhere, BlueprintReadWrite, Setter="SetSeed", Getter="GetSeed", Category="Cloner")
+	UE_DEPRECATED(5.5, "Moved to component and extensions")
+	UPROPERTY()
 	int32 Seed = 0;
 
-	/** Cloner color when unaffected by effectors, color will be passed down to the material (ParticleColor) */
-	UPROPERTY(EditAnywhere, BlueprintReadWrite, Setter, Getter, Category="Cloner")
+	UE_DEPRECATED(5.5, "Moved to component and extensions")
+	UPROPERTY()
 	FLinearColor Color = FLinearColor::White;
 
-	/** Switches materials to show effectors applied on this cloner based on their color property */
-	UPROPERTY(EditAnywhere, BlueprintReadWrite, Setter="SetVisualizeEffectors", Getter="GetVisualizeEffectors", Category="Cloner")
+	UE_DEPRECATED(5.5, "Moved to component and extensions")
+	UPROPERTY()
 	bool bVisualizeEffectors;
 
-	/** Indicates how we select the mesh to render on each clones */
-	UPROPERTY(EditAnywhere, BlueprintReadWrite, Setter="SetMeshRenderMode", Getter="GetMeshRenderMode", Category="Renderer")
+	UE_DEPRECATED(5.5, "Moved to component and extensions")
+	UPROPERTY()
 	ECEClonerMeshRenderMode MeshRenderMode = ECEClonerMeshRenderMode::Iterate;
 
-	/** Mode to indicate how clones facing is determined */
-	UPROPERTY(EditAnywhere, BlueprintReadWrite, Setter="SetMeshFacingMode", Getter="GetMeshFacingMode", Category="Renderer")
+	UE_DEPRECATED(5.5, "Moved to component and extensions")
+	UPROPERTY()
 	ENiagaraMeshFacingMode MeshFacingMode = ENiagaraMeshFacingMode::Default;
 
-	/** Whether clones cast shadow, disabling will result in better performance */
-	UPROPERTY(EditAnywhere, BlueprintReadWrite, Setter="SetMeshCastShadows", Getter="GetMeshCastShadows", Category="Renderer")
+	UE_DEPRECATED(5.5, "Moved to component and extensions")
+	UPROPERTY()
 	bool bMeshCastShadows = true;
 
-	/** When nothing is attached to the cloner, these meshes are used as default */
-	UPROPERTY(EditAnywhere, BlueprintReadWrite, Setter="SetDefaultMeshes", Getter="GetDefaultMeshes", BlueprintSetter="BP_SetDefaultMeshes", BlueprintGetter="BP_GetDefaultMeshes", Category="Renderer")
+	UE_DEPRECATED(5.5, "Moved to component and extensions")
+	UPROPERTY()
 	TArray<TObjectPtr<UStaticMesh>> DefaultMeshes;
 
-	/** Whether to override meshes materials with another material */
-	UPROPERTY(EditAnywhere, BlueprintReadWrite, Setter="SetUseOverrideMaterial", Getter="GetUseOverrideMaterial", Category="Renderer")
+	UE_DEPRECATED(5.5, "Moved to component and extensions")
+	UPROPERTY()
 	bool bUseOverrideMaterial = false;
 
-	/** The override materials that will be set instead of meshes materials, bVisualizeEffectors must be disabled */
-	UPROPERTY(EditAnywhere, BlueprintReadWrite, Setter, Getter, Category="Renderer", meta=(EditCondition="bUseOverrideMaterial", EditConditionHides))
+	UE_DEPRECATED(5.5, "Moved to component and extensions")
+	UPROPERTY()
 	TObjectPtr<UMaterialInterface> OverrideMaterial;
 
-	/**
-	 * Allow particles to react to surface by using the distance field,
-	 * ensure the mesh you want particle to collide with is tick enough or increase its distance field resolution scale in static mesh editor 
-	 */
-	UPROPERTY(EditAnywhere, BlueprintReadWrite, Setter="SetSurfaceCollisionEnabled", Getter="GetSurfaceCollisionEnabled", Category="Collisions")
+	UE_DEPRECATED(5.5, "Moved to component and extensions")
+	UPROPERTY()
 	bool bSurfaceCollisionEnabled = false;
 
-	/** Allow particles to react to other emitter particles, uses a neighbor grid to detect collision */
-	UPROPERTY(EditAnywhere, BlueprintReadWrite, Setter="SetParticleCollisionEnabled", Getter="GetParticleCollisionEnabled", Category="Collisions")
+	UE_DEPRECATED(5.5, "Moved to component and extensions")
+	UPROPERTY()
 	bool bParticleCollisionEnabled = false;
 
-	/** Recalculate accurate velocity after collision is updated */
-	UPROPERTY(EditAnywhere, BlueprintReadWrite, Setter="SetCollisionVelocityEnabled", Getter="GetCollisionVelocityEnabled", Category="Collisions", meta=(EditCondition="bParticleCollisionEnabled", EditConditionHides))
+	UE_DEPRECATED(5.5, "Moved to component and extensions")
+	UPROPERTY()
 	bool bCollisionVelocityEnabled = false;
 
-	/** Amount of iterations to improve particle collision results but affects performance */
-	UPROPERTY(EditAnywhere, BlueprintReadWrite, Setter, Getter, Category="Collisions", meta=(ClampMin="1", EditCondition="bParticleCollisionEnabled", EditConditionHides))
+	UE_DEPRECATED(5.5, "Moved to component and extensions")
+	UPROPERTY()
 	int32 CollisionIterations = 1;
 
-	/** Resolution of the neighbor grid to detect collision between particles */
-	UPROPERTY(EditAnywhere, BlueprintReadWrite, Setter, Getter, Category="Collisions", meta=(ClampMin="1", EditCondition="bParticleCollisionEnabled", EditConditionHides))
+	UE_DEPRECATED(5.5, "Moved to component and extensions")
+	UPROPERTY()
 	int32 CollisionGridResolution = 32;
 
-	/** Size of the neighbor grid to detect collision between particles, particles outside this grid will not have collisions */
-	UPROPERTY(EditAnywhere, BlueprintReadWrite, Setter, Getter, Category="Collisions", meta=(ClampMin="0", AllowPreserveRatio, EditCondition="bParticleCollisionEnabled", EditConditionHides))
+	UE_DEPRECATED(5.5, "Moved to component and extensions")
+	UPROPERTY()
 	FVector CollisionGridSize = FVector(5000.f);
 
-	/** Collision radius calculation mode */
-	UPROPERTY(EditAnywhere, BlueprintReadWrite, Setter, Getter, Category="Collisions", meta=(EditCondition="bParticleCollisionEnabled", EditConditionHides))
+	UE_DEPRECATED(5.5, "Moved to component and extensions")
+	UPROPERTY()
 	ECEClonerCollisionRadiusMode CollisionRadiusMode = ECEClonerCollisionRadiusMode::ExtentLength;
 
-	/** Radius expected around each particle for collision, order matches attachment index */
-	UPROPERTY(EditAnywhere, BlueprintReadWrite, Getter, Category="Collisions", EditFixedSize, meta=(ClampMin="0", EditFixedOrder, EditCondition="bParticleCollisionEnabled && CollisionRadiusMode == ECEClonerCollisionRadiusMode::Manual", EditConditionHides))
+	UE_DEPRECATED(5.5, "Moved to component and extensions")
+	UPROPERTY()
 	TArray<float> CollisionRadii;
 
-	/** Minimum particle mass, used for collisions to push apart */
-	UPROPERTY(EditAnywhere, BlueprintReadWrite, Setter, Getter, Category="Collisions", meta=(ClampMin="1"))
+	UE_DEPRECATED(5.5, "Moved to component and extensions")
+	UPROPERTY()
 	float MassMin = 1.f;
 
-	/** Maximum particle mass, used for collisions to push apart */
-	UPROPERTY(EditAnywhere, BlueprintReadWrite, Setter, Getter, Category="Collisions", meta=(ClampMin="1"))
+	UE_DEPRECATED(5.5, "Moved to component and extensions")
+	UPROPERTY()
 	float MassMax = 2.f;
 
-	/** How many times do we spawn clones */
-	UPROPERTY(EditAnywhere, BlueprintReadWrite, Setter="SetSpawnLoopMode", Getter="GetSpawnLoopMode", DisplayName="LoopMode", Category="Spawn")
+	UE_DEPRECATED(5.5, "Moved to component and extensions")
+	UPROPERTY()
 	ECEClonerSpawnLoopMode SpawnLoopMode = ECEClonerSpawnLoopMode::Once;
 
-	/** Amount of spawn iterations for clones */
-	UPROPERTY(EditAnywhere, BlueprintReadWrite, Setter="SetSpawnLoopIterations", Getter="GetSpawnLoopIterations", DisplayName="LoopIterations", Category="Spawn", meta=(ClampMin="1", EditCondition="SpawnLoopMode == ECEClonerSpawnLoopMode::Multiple", EditConditionHides))
+	UE_DEPRECATED(5.5, "Moved to component and extensions")
+	UPROPERTY()
 	int32 SpawnLoopIterations = 1;
 
-	/** Interval/Duration of spawn for clones */
-	UPROPERTY(EditAnywhere, BlueprintReadWrite, Setter="SetSpawnLoopInterval", Getter="GetSpawnLoopInterval", DisplayName="LoopInterval", Category="Spawn", meta=(ClampMin="0"))
+	UE_DEPRECATED(5.5, "Moved to component and extensions")
+	UPROPERTY()
 	float SpawnLoopInterval = 1.f;
 
-	/** How does spawn occurs */
-	UPROPERTY(EditAnywhere, BlueprintReadWrite, Setter="SetSpawnBehaviorMode", Getter="GetSpawnBehaviorMode", DisplayName="BehaviorMode", Category="Spawn")
+	UE_DEPRECATED(5.5, "Moved to component and extensions")
+	UPROPERTY()
 	ECEClonerSpawnBehaviorMode SpawnBehaviorMode = ECEClonerSpawnBehaviorMode::Instant;
 
-	/** How many clones to spawn each seconds */
-	UPROPERTY(EditAnywhere, BlueprintReadWrite, Setter="SetSpawnRate", Getter="GetSpawnRate", DisplayName="Rate", Category="Spawn", meta=(ClampMin="0", EditCondition="SpawnBehaviorMode == ECEClonerSpawnBehaviorMode::Rate", EditConditionHides))
+	UE_DEPRECATED(5.5, "Moved to component and extensions")
+	UPROPERTY()
 	float SpawnRate = 1.f;
 
-	/** Do we destroy the clones after a specific duration */
-	UPROPERTY(EditAnywhere, BlueprintReadWrite, Setter="SetLifetimeEnabled", Getter="GetLifetimeEnabled", DisplayName="Enabled", Category="Lifetime")
+	UE_DEPRECATED(5.5, "Moved to component and extensions")
+	UPROPERTY()
 	bool bLifetimeEnabled = false;
 
-	/** Minimum lifetime for a clone */
-	UPROPERTY(EditAnywhere, BlueprintReadWrite, Setter="SetLifetimeMin", Getter="GetLifetimeMin", DisplayName="Min", Category="Lifetime", meta=(ClampMin="0", EditCondition="bLifetimeEnabled", EditConditionHides))
+	UE_DEPRECATED(5.5, "Moved to component and extensions")
+	UPROPERTY()
 	float LifetimeMin = 0.25f;
 
-	/** Maximum lifetime for a clone */
-	UPROPERTY(EditAnywhere, BlueprintReadWrite, Setter="SetLifetimeMax", Getter="GetLifetimeMax", DisplayName="Max", Category="Lifetime", meta=(ClampMin="0", EditCondition="bLifetimeEnabled", EditConditionHides))
+	UE_DEPRECATED(5.5, "Moved to component and extensions")
+	UPROPERTY()
 	float LifetimeMax = 1.5f;
 
-	/** Enable scale by lifetime */
-	UPROPERTY(EditAnywhere, BlueprintReadWrite, Setter="SetLifetimeScaleEnabled", Getter="GetLifetimeScaleEnabled", DisplayName="ScaleEnabled", Category="Lifetime", meta=(EditCondition="bLifetimeEnabled", EditConditionHides))
+	UE_DEPRECATED(5.5, "Moved to component and extensions")
+	UPROPERTY()
 	bool bLifetimeScaleEnabled = false;
 
-	/** Used to expose the scale curve editor in details panel */
-	UPROPERTY(Transient, DuplicateTransient, TextExportTransient)
-	TWeakObjectPtr<UNiagaraDataInterfaceCurve> LifetimeScaleCurveDIWeak;
-
-	UPROPERTY(Setter="SetLifetimeScaleCurve", Getter="GetLifetimeScaleCurve")
+	UE_DEPRECATED(5.5, "Moved to component and extensions")
+	UPROPERTY()
 	FRichCurve LifetimeScaleCurve;
 
-	/** Name of the layout to use */
-	UPROPERTY(EditAnywhere, BlueprintReadWrite, Setter="SetLayoutName", Getter="GetLayoutName", Category="Layout", meta=(GetOptions="GetClonerLayoutNames"))
+	UE_DEPRECATED(5.5, "Moved to component and extensions")
+	UPROPERTY()
 	FName LayoutName = NAME_None;
 
-	/** Active layout */
-	UPROPERTY(VisibleInstanceOnly, BlueprintReadOnly, Instanced, Transient, DuplicateTransient, Category="Layout", meta=(DisplayAfter="LayoutName"))
-	TObjectPtr<UCEClonerLayoutBase> ActiveLayout;
-
-	/** Invert progress behaviour */
-	UPROPERTY(EditAnywhere, BlueprintReadWrite, Setter="SetInvertProgress", Getter="GetInvertProgress", Category="Progress")
-	bool bInvertProgress = false;
-
-	/** Changes visibility of instances based on the total count, 1.f = 100% = all instances visible */
-	UPROPERTY(EditAnywhere, BlueprintReadWrite, Setter="SetProgress", Getter="GetProgress", Category="Progress", meta=(ClampMin="0", ClampMax="1"))
-	float Progress = 1.f;
-
-	/** Enable steps to add delta variation on each clone instance */
-	UPROPERTY(EditAnywhere, BlueprintReadWrite, Setter="SetDeltaStepEnabled", Getter="GetDeltaStepEnabled", DisplayName="Enabled", Category="Step")
+	UE_DEPRECATED(5.5, "Moved to component and extensions")
+	UPROPERTY()
 	bool bDeltaStepEnabled = false;
 
-	/** Amount of rotation difference between one step and the next one */
-	UPROPERTY(EditAnywhere, BlueprintReadWrite, Setter="SetDeltaStepRotation", Getter="GetDeltaStepRotation", DisplayName="Rotation", Category="Step", meta=(EditCondition="bDeltaStepEnabled", EditConditionHides))
+	UE_DEPRECATED(5.5, "Moved to component and extensions")
+	UPROPERTY()
 	FRotator DeltaStepRotation = FRotator(0.f);
 
-	/** Amount of scale difference between one step and the next one */
-	UPROPERTY(EditAnywhere, BlueprintReadWrite, Setter="SetDeltaStepScale", Getter="GetDeltaStepScale", DisplayName="Scale", Category="Step", meta=(AllowPreserveRatio, Delta="0.0001", EditCondition="bDeltaStepEnabled", EditConditionHides))
+	UE_DEPRECATED(5.5, "Moved to component and extensions")
+	UPROPERTY()
 	FVector DeltaStepScale = FVector(0.f);
 
-	/** Use random range transforms for each clones */
-	UPROPERTY(EditAnywhere, BlueprintReadWrite, Setter="SetRangeEnabled", Getter="GetRangeEnabled", DisplayName="Enabled", Category="Range")
+	UE_DEPRECATED(5.5, "Moved to component and extensions")
+	UPROPERTY()
 	bool bRangeEnabled = false;
 
-	UPROPERTY(EditAnywhere, BlueprintReadWrite, Setter="SetRangeOffsetMin", Getter="GetRangeOffsetMin", DisplayName="OffsetMin", Category="Range", meta=(EditCondition="bRangeEnabled", EditConditionHides))
+	UE_DEPRECATED(5.5, "Moved to component and extensions")
+	UPROPERTY()
 	FVector RangeOffsetMin = FVector::ZeroVector;
 
-	UPROPERTY(EditAnywhere, BlueprintReadWrite, Setter="SetRangeOffsetMax", Getter="GetRangeOffsetMax", DisplayName="OffsetMax", Category="Range", meta=(EditCondition="bRangeEnabled", EditConditionHides))
+	UE_DEPRECATED(5.5, "Moved to component and extensions")
+	UPROPERTY()
 	FVector RangeOffsetMax = FVector::ZeroVector;
 
-	UPROPERTY(EditAnywhere, BlueprintReadWrite, Setter="SetRangeRotationMin", Getter="GetRangeRotationMin", DisplayName="RotationMin", Category="Range", meta=(EditCondition="bRangeEnabled", EditConditionHides))
+	UE_DEPRECATED(5.5, "Moved to component and extensions")
+	UPROPERTY()
 	FRotator RangeRotationMin = FRotator::ZeroRotator;
 
-	UPROPERTY(EditAnywhere, BlueprintReadWrite, Setter="SetRangeRotationMax", Getter="GetRangeRotationMax", DisplayName="RotationMax", Category="Range", meta=(EditCondition="bRangeEnabled", EditConditionHides))
+	UE_DEPRECATED(5.5, "Moved to component and extensions")
+	UPROPERTY()
 	FRotator RangeRotationMax = FRotator::ZeroRotator;
 
-	UPROPERTY(EditAnywhere, BlueprintReadWrite, Setter="SetRangeScaleUniform", Getter="GetRangeScaleUniform", DisplayName="ScaleUniformEnabled", Category="Range", meta=(EditCondition="bRangeEnabled", EditConditionHides))
+	UE_DEPRECATED(5.5, "Moved to component and extensions")
+	UPROPERTY()
 	bool bRangeScaleUniform = false;
 
-	UPROPERTY(EditAnywhere, BlueprintReadWrite, Setter="SetRangeScaleMin", Getter="GetRangeScaleMin", DisplayName="ScaleMin", Category="Range", meta=(AllowPreserveRatio, Delta="0.0001", ClampMin="0", EditCondition="bRangeEnabled && !bRangeScaleUniform", EditConditionHides))
+	UE_DEPRECATED(5.5, "Moved to component and extensions")
+	UPROPERTY()
 	FVector RangeScaleMin = FVector::OneVector;
 
-	UPROPERTY(EditAnywhere, BlueprintReadWrite, Setter="SetRangeScaleMax", Getter="GetRangeScaleMax", DisplayName="ScaleMax", Category="Range", meta=(AllowPreserveRatio, Delta="0.0001", ClampMin="0", EditCondition="bRangeEnabled && !bRangeScaleUniform", EditConditionHides))
+	UE_DEPRECATED(5.5, "Moved to component and extensions")
+	UPROPERTY()
 	FVector RangeScaleMax = FVector::OneVector;
 
-	UPROPERTY(EditAnywhere, BlueprintReadWrite, Setter="SetRangeScaleUniformMin", Getter="GetRangeScaleUniformMin", DisplayName="ScaleMin", Category="Range", meta=(Delta="0.0001", ClampMin="0", EditCondition="bRangeEnabled && bRangeScaleUniform", EditConditionHides))
+	UE_DEPRECATED(5.5, "Moved to component and extensions")
+	UPROPERTY()
 	float RangeScaleUniformMin = 1.f;
 
-	UPROPERTY(EditAnywhere, BlueprintReadWrite, Setter="SetRangeScaleUniformMax", Getter="GetRangeScaleUniformMax", DisplayName="ScaleMax", Category="Range", meta=(Delta="0.0001", ClampMin="0", EditCondition="bRangeEnabled && bRangeScaleUniform", EditConditionHides))
+	UE_DEPRECATED(5.5, "Moved to component and extensions")
+	UPROPERTY()
 	float RangeScaleUniformMax = 1.f;
 
-#if WITH_EDITORONLY_DATA
-	/** Toggle the sprite to visualize and click on this cloner */
-	UPROPERTY(EditAnywhere, BlueprintReadWrite, AdvancedDisplay, Category="Cloner", meta=(AllowPrivateAccess="true"))
-	bool bVisualizerSpriteVisible = true;
-
-	/** Reduces the r.TSR.ShadingRejection.Flickering.Period from 3 (default) to 1 if enabled to avoid ghosting artifacts when moving */
-	UPROPERTY(EditAnywhere, BlueprintReadWrite, AdvancedDisplay, Category="Cloner", meta=(AllowPrivateAccess="true"))
-	bool bReduceMotionGhosting = false;
-#endif
-
-private:
-#if WITH_EDITOR
-	/** Hide outline selection when this cloner only is selected in viewport */
-	void OnEditorSelectionChanged(UObject* InSelection);
-
-	TOptional<bool> UseSelectionOutline;
-#endif
-
-	/** Gets all layout names available */
-	UFUNCTION()
-	TArray<FString> GetClonerLayoutNames() const;
-
-	template<
-		typename InLayoutClass
-		UE_REQUIRES(TIsDerivedFrom<InLayoutClass, UCEClonerLayoutBase>::Value)>
-	InLayoutClass* FindOrAddLayout()
-	{
-		return Cast<InLayoutClass>(FindOrAddLayout(InLayoutClass::StaticClass()));
-	}
-
-	/** Find or add a layout by its class */
-	UCEClonerLayoutBase* FindOrAddLayout(TSubclassOf<UCEClonerLayoutBase> InClass);
-
-	/** Find or add a layout by its name */
-	UCEClonerLayoutBase* FindOrAddLayout(FName InLayoutName);
-
-	/** Initiate and perform operation */
-	void InitializeCloner();
-
-	/** Cloner niagara component */
+	UE_DEPRECATED(5.5, "Moved to component and extensions")
 	UPROPERTY()
-	TObjectPtr<UCEClonerComponent> ClonerComponent;
+	bool bInvertProgress = false;
 
-	/** Previously used layout instances cached */
+	UE_DEPRECATED(5.5, "Moved to component and extensions")
 	UPROPERTY()
-	TMap<FName, TObjectPtr<UCEClonerLayoutBase>> LayoutInstances;
+	float Progress = 1.f;
 
-	/** Effectors linked to this cloner */
-	UPROPERTY(EditInstanceOnly, Category="Cloner", meta=(DisplayName="Effectors"))
+	UE_DEPRECATED(5.5, "Moved to component and extensions")
+	UPROPERTY()
 	TArray<TWeakObjectPtr<ACEEffectorActor>> EffectorsWeak;
 
-	/** Copy of effectors linked to this cloner to compare diffs */
-	UPROPERTY(Transient, DuplicateTransient, TextExportTransient, NonTransactional)
-	TSet<TWeakObjectPtr<ACEEffectorActor>> EffectorsInternalWeak;
-
-	float TreeUpdateDeltaTime = 0.f;
-
-	bool bNeedsRefresh = false;
+	UE_DEPRECATED(5.5, "Moved to component and extensions")
+	UPROPERTY()
+	TMap<FName, TObjectPtr<UCEClonerLayoutBase>> LayoutInstances;
 
-	bool bClonerInitialized = false;
+#if WITH_EDITORONLY_DATA
+	UE_DEPRECATED(5.5, "Moved to component and extensions")
+	UPROPERTY()
+	bool bVisualizerSpriteVisible = true;
 
-#if WITH_EDITOR
 	bool bSpawnDefaultActorAttached = false;
-
-	/** Used for PECP */
-	static const TCEPropertyChangeDispatcher<ACEClonerActor> PropertyChangeDispatcher;
 #endif
 };
diff --git a/Engine/Plugins/Experimental/ClonerEffector/Source/ClonerEffector/Public/Cloner/CEClonerComponent.h b/Engine/Plugins/Experimental/ClonerEffector/Source/ClonerEffector/Public/Cloner/CEClonerComponent.h
--- a/Engine/Plugins/Experimental/ClonerEffector/Source/ClonerEffector/Public/Cloner/CEClonerComponent.h
+++ b/Engine/Plugins/Experimental/ClonerEffector/Source/ClonerEffector/Public/Cloner/CEClonerComponent.h
@@ -3,28 +3,30 @@
 #pragma once
 
 #include "CEClonerEffectorShared.h"
+#include "CEClonerMeshBuilder.h"
+#include "CEPropertyChangeDispatcher.h"
 #include "Layouts/CEClonerLayoutBase.h"
 #include "NiagaraComponent.h"
 #include "CEClonerComponent.generated.h"
 
+class UCEClonerExtensionBase;
 class UCEClonerLayoutBase;
 class UMaterialInterface;
 struct FNiagaraMeshMaterialOverride;
 
-UCLASS(MinimalAPI, Within=CEClonerActor, HideCategories=(Niagara))
+UCLASS(MinimalAPI
+	, BlueprintType
+	, DisplayName = "Motion Design Cloner Component"
+	, AutoExpandCategories=(Cloner, Layout)
+	, HideCategories=(Niagara, Activation, Lighting, Attachment, Randomness, Parameters, Warmup, Compilation, Navigation, Tags, LOD, TextureStreaming, Mobile, RayTracing, AssetUserData, Cooking, HLOD, Rendering))
 class UCEClonerComponent : public UNiagaraComponent
 {
 	GENERATED_BODY()
 
-public:
-	/** Called when meshes have been updated */
-	DECLARE_MULTICAST_DELEGATE_OneParam(FOnClonerMeshUpdated, UCEClonerComponent* /** ClonerComponent */)
-	static FOnClonerMeshUpdated OnClonerMeshUpdated;
-
-	/** Called when cloner layout system is loaded */
-	DECLARE_MULTICAST_DELEGATE_TwoParams(FOnClonerSystemLoaded, UCEClonerComponent* /** ClonerComponent */, UCEClonerLayoutBase* /** InLayout */)
-	static FOnClonerSystemLoaded OnClonerSystemLoaded;
+	friend class ACEClonerActor;
+	friend class FCEEditorClonerComponentDetailCustomization;
 
+public:
 	/** Only materials transient or part of the content folder can be dirtied, engine or plugins cannot */
 	static bool IsMaterialDirtyable(const UMaterialInterface* InMaterial);
 
@@ -36,20 +38,127 @@ public:
 	static void ShowMaterialWarning(int32 InMaterialCount);
 #endif
 
+	TMulticastDelegateRegistration<void(UCEClonerComponent*)>& OnClonerMeshUpdated()
+	{
+		return OnClonerMeshUpdatedDelegate;
+	}
+
+	TMulticastDelegateRegistration<void(UCEClonerComponent*, UCEClonerLayoutBase*)>& OnClonerLayoutLoaded()
+	{
+		return OnClonerLayoutLoadedDelegate;
+	}
+
+	TMulticastDelegateRegistration<void(UCEClonerComponent*)>& OnClonerInitialized()
+	{
+		return OnClonerInitializedDelegate;
+	}
+
 	UCEClonerComponent();
 
+	UFUNCTION(BlueprintCallable, Category="Cloner")
+	CLONEREFFECTOR_API void SetEnabled(bool bInEnable);
+
+	UFUNCTION(BlueprintPure, Category="Cloner")
+	bool GetEnabled() const
+	{
+		return bEnabled;
+	}
+
+	UFUNCTION(BlueprintCallable, Category="Cloner")
+	CLONEREFFECTOR_API void SetTreeUpdateInterval(float InInterval);
+
+	UFUNCTION(BlueprintPure, Category="Cloner")
+	float GetTreeUpdateInterval() const
+	{
+		return TreeUpdateInterval;
+	}
+
+	UFUNCTION(BlueprintCallable, Category="Cloner")
+	CLONEREFFECTOR_API void SetSeed(int32 InSeed);
+
+	UFUNCTION(BlueprintPure, Category="Cloner")
+	int32 GetSeed() const
+	{
+		return Seed;
+	}
+
+	UFUNCTION(BlueprintCallable, Category="Cloner")
+	CLONEREFFECTOR_API void SetLayoutName(FName InLayoutName);
+
+	UFUNCTION(BlueprintPure, Category="Cloner")
+	FName GetLayoutName() const
+	{
+		return LayoutName;
+	}
+
+	UFUNCTION(BlueprintPure, Category="Cloner")
+	UCEClonerLayoutBase* GetActiveLayout() const
+	{
+		return ActiveLayout;
+	}
+
+	template<
+		typename InLayoutClass
+		UE_REQUIRES(TIsDerivedFrom<InLayoutClass, UCEClonerLayoutBase>::Value)>
+	bool IsActiveLayout() const
+	{
+		if (const UCEClonerLayoutBase* CurrentLayout = GetActiveLayout())
+		{
+			return CurrentLayout->GetClass() == InLayoutClass::StaticClass();
+		}
+
+		return false;
+	}
+
+	template<
+		typename InLayoutClass
+		UE_REQUIRES(TIsDerivedFrom<InLayoutClass, UCEClonerLayoutBase>::Value)>
+	InLayoutClass* GetActiveLayout() const
+	{
+		return Cast<InLayoutClass>(GetActiveLayout());
+	}
+
+#if WITH_EDITOR
+	UFUNCTION(BlueprintCallable, Category="Cloner")
+	CLONEREFFECTOR_API void SetVisualizerSpriteVisible(bool bInVisible);
+
+	UFUNCTION(BlueprintPure, Category="Cloner")
+	bool GetVisualizerSpriteVisible() const
+	{
+		return bVisualizerSpriteVisible;
+	}
+#endif
+
+	/** Returns the number of meshes this cloner currently handles */
+	UFUNCTION(BlueprintPure, Category="Cloner")
+	CLONEREFFECTOR_API int32 GetMeshCount() const;
+
+	/** Returns the number of root attachment currently on this cloner */
+	UFUNCTION(BlueprintPure, Category="Cloner")
+	int32 GetAttachmentCount() const;
+
+#if WITH_EDITOR
+	/** This will force an update of the cloner attachment tree */
+	UFUNCTION(CallInEditor, Category="Cloner")
+	void ForceUpdateCloner();
+
+	/** This will create a new default actor attached to this cloner */
+	UFUNCTION(CallInEditor, Category="Cloner")
+	void CreateDefaultActorAttached();
+#endif
+
+	/** Will force a system update to refresh user parameters */
+	void RequestClonerUpdate(bool bInImmediate = false);
+
 	/**
 	 * Triggers an update of the attachment tree to detect updated items
 	 * If reset is true, clears the attachment tree and rebuilds it otherwise diff update
 	 */
 	void UpdateClonerAttachmentTree(bool bInReset = false);
 
-	/** Called to trigger an update of cloner rendering state if tree */
+	/** Called to trigger an update of cloner rendering state tree */
 	void UpdateClonerRenderState();
 
-	/** Sets the layout to use for this cloner simulation */
-	bool SetClonerActiveLayout(UCEClonerLayoutBase* InLayout);
-
 	UCEClonerLayoutBase* GetClonerActiveLayout() const
 	{
 		return ActiveLayout;
@@ -58,17 +167,52 @@ public:
 	/** Forces a refresh of the active system parameters in niagara store */
 	void RefreshUserParameters() const;
 
-	void SetUseOverrideMeshesMaterial(bool bInOverride);
+	/** Forces a refresh of the meshes used */
+	void RefreshClonerMeshes();
+
+	template<
+		typename InExtensionClass>
+	InExtensionClass* GetExtension() const
+	{
+		return Cast<InExtensionClass>(GetExtension(InExtensionClass::StaticClass()));
+	}
+
+	UCEClonerExtensionBase* GetExtension(TSubclassOf<UCEClonerExtensionBase> InExtensionClass) const;
+
+	UCEClonerExtensionBase* GetExtension(FName InExtensionName) const;
 
-	void SetOverrideMeshesMaterial(UMaterialInterface* InMaterial);
+	TConstArrayView<TObjectPtr<UCEClonerExtensionBase>> GetActiveExtensions() const
+	{
+		return ActiveExtensions;
+	}
 
 protected:
+	/** Called when meshes have been updated */
+	DECLARE_MULTICAST_DELEGATE_OneParam(FOnClonerMeshUpdated, UCEClonerComponent* /** ClonerComponent */)
+	FOnClonerMeshUpdated OnClonerMeshUpdatedDelegate;
+
+	/** Called when new cloner layout is loaded */
+	DECLARE_MULTICAST_DELEGATE_TwoParams(FOnClonerLayoutLoaded, UCEClonerComponent* /** ClonerComponent */, UCEClonerLayoutBase* /** InLayout */)
+	FOnClonerLayoutLoaded OnClonerLayoutLoadedDelegate;
+
+	/** Called when cloner is initialized */
+	DECLARE_MULTICAST_DELEGATE_OneParam(FOnClonerInitialized, UCEClonerComponent* /** ClonerComponent */)
+	FOnClonerInitialized OnClonerInitializedDelegate;
+
+	//~ Begin UActorComponent
+	/** Tick needed to keep the attachment tree updated */
+	virtual void TickComponent(float InDeltaTime, ELevelTick InTickType, FActorComponentTickFunction* InThisTickFunction) override;
+	//~ End UActorComponent
+
 	//~ Begin UObject
+	virtual void PostInitProperties() override;
 	virtual void PostLoad() override;
+#if WITH_EDITOR
+	virtual void PostEditUndo() override;
+	virtual void PostEditChangeProperty(FPropertyChangedEvent& InPropertyChangedEvent) override;
+#endif
 	//~ End UObject
 
-	void OnRenderStateDirty(UActorComponent& InComponent);
-
 	void UpdateAttachmentTree();
 	void UpdateActorAttachment(AActor* InActor, AActor* InParent);
 
@@ -118,32 +262,123 @@ protected:
 	/** Checks that all root static meshes are valid */
 	bool IsAllMergedMeshesValid() const;
 
-	/** Called when material override properties are changed */
-	void OnOverrideMeshesMaterialChanged();
+	/** Sets the layout to use for this cloner simulation */
+	void SetClonerActiveLayout(UCEClonerLayoutBase* InLayout);
 
-	/** Get the niagara materials override if enabled */
-	TArray<FNiagaraMeshMaterialOverride> GetOverrideMeshesMaterials() const;
+	void OnActiveLayoutLoaded(UCEClonerLayoutBase* InLayout, bool bInSuccess);
 
-	/** Attachment tree view */
-	UPROPERTY(Transient, NonTransactional)
-	FCEClonerAttachmentTree ClonerTree;
+	void ActivateLayout(UCEClonerLayoutBase* InLayout);
 
-	/** Current active layout */
-	UPROPERTY(Transient, NonTransactional)
+	void OnActiveLayoutChanged();
+
+	/** Is this cloner enabled/disabled */
+	UPROPERTY(EditInstanceOnly, BlueprintReadWrite, Setter="SetEnabled", Getter="GetEnabled", Category="Cloner")
+	bool bEnabled = true;
+
+	/** Interval to update the attachment tree and update the cloner meshes, 0 means each tick */
+	UPROPERTY(EditInstanceOnly, BlueprintReadWrite, Setter, Getter, Category="Cloner", AdvancedDisplay, meta=(ClampMin="0"))
+	float TreeUpdateInterval = 0.2f;
+
+	/** Cloner instance seed for random deterministic patterns */
+	UPROPERTY(EditInstanceOnly, BlueprintReadWrite, Setter, Getter, Category="Cloner")
+	int32 Seed = 0;
+
+	/** Name of the layout to use */
+	UPROPERTY(EditInstanceOnly, BlueprintReadWrite, Setter, Getter, Category="Layout", meta=(GetOptions="GetClonerLayoutNames"))
+	FName LayoutName = NAME_None;
+
+	/** Active layout used */
+	UPROPERTY(VisibleInstanceOnly, BlueprintReadOnly, Instanced, Transient, DuplicateTransient, NonTransactional, Category="Layout", meta=(DisplayAfter="LayoutName"))
 	TObjectPtr<UCEClonerLayoutBase> ActiveLayout;
 
-	/** Attachment items that are dirty and need an update */
-	TSet<FCEClonerAttachmentItem*> DirtyItemAttachments;
+	/** Active Extensions on this layout */
+	UPROPERTY(VisibleInstanceOnly, BlueprintReadOnly, Instanced, Transient, DuplicateTransient, NonTransactional, Category="Layout")
+	TArray<TObjectPtr<UCEClonerExtensionBase>> ActiveExtensions;
+
+	/** Layout instances cached */
+	UPROPERTY()
+	TArray<TObjectPtr<UCEClonerLayoutBase>> LayoutInstances;
+
+	/** Layout extensions instances cached */
+	UPROPERTY()
+	TArray<TObjectPtr<UCEClonerExtensionBase>> ExtensionInstances;
+
+#if WITH_EDITORONLY_DATA
+	/** Toggle the sprite to visualize and click on this cloner */
+	UPROPERTY(EditInstanceOnly, BlueprintReadWrite, AdvancedDisplay, Category="Cloner", meta=(AllowPrivateAccess="true"))
+	bool bVisualizerSpriteVisible = true;
+#endif
 
-	/** Override meshes material */
-	TWeakObjectPtr<UMaterialInterface> ClonerMeshesOverrideMaterial;
+private:
+	static constexpr TCHAR SpriteTexturePath[] = TEXT("/Script/Engine.Texture2D'/ClonerEffector/Textures/T_ClonerIcon.T_ClonerIcon'");
 
-	/** Use override meshes material */
-	bool bUseClonerMeshesOverrideMaterial = false;
+	/** Initiate and perform operation */
+	void InitializeCloner();
+	void TickCloner(float InDelta);
+
+	void OnEnabledChanged();
+	void OnClonerEnabled();
+	void OnClonerDisabled();
+	void OnSeedChanged();
+	void OnLayoutNameChanged();
+
+#if WITH_EDITOR
+	void OnVisualizerSpriteVisibleChanged();
+#endif
+
+	template<
+		typename InLayoutClass
+		UE_REQUIRES(TIsDerivedFrom<InLayoutClass, UCEClonerLayoutBase>::Value)>
+	InLayoutClass* FindOrAddLayout()
+	{
+		return Cast<InLayoutClass>(FindOrAddLayout(InLayoutClass::StaticClass()));
+	}
+
+	/** Find or add a layout by its class */
+	UCEClonerLayoutBase* FindOrAddLayout(TSubclassOf<UCEClonerLayoutBase> InClass);
+
+	/** Find or add a layout by its name */
+	UCEClonerLayoutBase* FindOrAddLayout(FName InLayoutName);
+
+	template<
+		typename InExtensionClass
+		UE_REQUIRES(TIsDerivedFrom<InExtensionClass, UCEClonerExtensionBase>::Value)>
+	InExtensionClass* FindOrAddExtension()
+	{
+		return Cast<InExtensionClass>(FindOrAddExtension(InExtensionClass::StaticClass()));
+	}
+
+	/** Find or add an extension by its class */
+	UCEClonerExtensionBase* FindOrAddExtension(TSubclassOf<UCEClonerExtensionBase> InClass);
+
+	/** Find or add an extension by its name */
+	UCEClonerExtensionBase* FindOrAddExtension(FName InExtensionName);
+
+	/** Gets all layout names available */
+	UFUNCTION()
+	TArray<FString> GetClonerLayoutNames() const;
+
+	/** Attachment tree view */
+	UPROPERTY(Transient, NonTransactional)
+	FCEClonerAttachmentTree ClonerTree;
+
+	UPROPERTY(Transient, NonTransactional)
+	FCEClonerMeshBuilder MeshBuilder;
 
 	/** Asset meshes needs update */
 	bool bClonerMeshesDirty = true;
 
 	/** State of the baked dynamic and static mesh creation */
-	bool bClonerMeshesUpdating = false;
+	std::atomic<bool> bClonerMeshesUpdating = false;
+
+	float TreeUpdateDeltaTime = 0.f;
+
+	bool bNeedsRefresh = false;
+
+	bool bClonerInitialized = false;
+
+#if WITH_EDITOR
+	/** Used for PECP */
+	static const TCEPropertyChangeDispatcher<UCEClonerComponent> PropertyChangeDispatcher;
+#endif
 };
\ No newline at end of file
diff --git a/Engine/Plugins/Experimental/ClonerEffector/Source/ClonerEffector/Public/Cloner/CEClonerMeshBuilder.h b/Engine/Plugins/Experimental/ClonerEffector/Source/ClonerEffector/Public/Cloner/CEClonerMeshBuilder.h
new file mode 100644
--- /dev/null
+++ b/Engine/Plugins/Experimental/ClonerEffector/Source/ClonerEffector/Public/Cloner/CEClonerMeshBuilder.h
@@ -0,0 +1,73 @@
+// Copyright Epic Games, Inc. All Rights Reserved.
+
+#pragma once
+
+#include "DynamicMesh/DynamicMesh3.h"
+#include "UObject/WeakObjectPtrTemplates.h"
+#include "CEClonerMeshBuilder.generated.h"
+
+class UBrushComponent;
+class UDynamicMesh;
+class UDynamicMeshComponent;
+class UInstancedStaticMeshComponent;
+class UMaterialInterface;
+class UProceduralMeshComponent;
+class USkeletalMeshComponent;
+class USplineMeshComponent;
+class UStaticMesh;
+class UStaticMeshComponent;
+
+/** Struct used to build a mesh based out of other meshes class */
+USTRUCT()
+struct FCEClonerMeshBuilder
+{
+	GENERATED_BODY()
+
+	FCEClonerMeshBuilder();
+
+	/** Resets the builder and clears data */
+	void Reset();
+
+	/** Appends supported components within actor */
+	int32 AppendActor(const AActor* InActor);
+
+	/** Appends Static Mesh component */
+	bool AppendComponent(const UStaticMeshComponent* InComponent, const FTransform& InTransform = FTransform::Identity);
+
+	/** Appends Procedural Mesh component */
+	bool AppendComponent(UProceduralMeshComponent* InComponent, const FTransform& InTransform = FTransform::Identity);
+
+	/** Appends Brush component */
+	bool AppendComponent(UBrushComponent* InComponent, const FTransform& InTransform = FTransform::Identity);
+
+	/** Appends Skeletal Mesh component */
+	bool AppendComponent(const USkeletalMeshComponent* InComponent, const FTransform& InTransform = FTransform::Identity);
+
+	/** Appends Dynamic Mesh component */
+	bool AppendComponent(UDynamicMeshComponent* InComponent, const FTransform& InTransform = FTransform::Identity);
+
+	/** Appends Instanced Static Mesh component */
+	bool AppendComponent(UInstancedStaticMeshComponent* InComponent, const FTransform& InTransform = FTransform::Identity);
+
+	/** Appends Spline Mesh component */
+	bool AppendComponent(USplineMeshComponent* InComponent, const FTransform& InTransform = FTransform::Identity);
+
+	/** Appends a dynamic mesh */
+	bool AppendMesh(const UDynamicMesh* InMesh, const TArray<TWeakObjectPtr<UMaterialInterface>>& InMaterials, const FTransform& InTransform = FTransform::Identity);
+
+	/** Builds a dynamic mesh by merging the mesh data imported */
+	bool BuildDynamicMesh(UDynamicMesh* OutMesh, TArray<TWeakObjectPtr<UMaterialInterface>>& OutMaterials);
+
+	/** Builds a static mesh by merging the mesh data imported */
+	bool BuildStaticMesh(UStaticMesh* OutMesh, TArray<TWeakObjectPtr<UMaterialInterface>>& OutMaterials);
+
+private:
+	bool AppendPrimitiveComponent(UPrimitiveComponent* InComponent, const FTransform& InTransform);
+
+	TArray<UE::Geometry::FDynamicMesh3> ConvertedMeshes;
+
+	TArray<TWeakObjectPtr<UMaterialInterface>> OutputMeshMaterials;
+
+	UPROPERTY()
+	TObjectPtr<UDynamicMesh> OutputDynamicMesh;
+};
diff --git a/Engine/Plugins/Experimental/ClonerEffector/Source/ClonerEffector/Public/Cloner/Extensions/CEClonerEmitterSpawnExtension.h b/Engine/Plugins/Experimental/ClonerEffector/Source/ClonerEffector/Public/Cloner/Extensions/CEClonerEmitterSpawnExtension.h
new file mode 100644
--- /dev/null
+++ b/Engine/Plugins/Experimental/ClonerEffector/Source/ClonerEffector/Public/Cloner/Extensions/CEClonerEmitterSpawnExtension.h
@@ -0,0 +1,111 @@
+// Copyright Epic Games, Inc. All Rights Reserved.
+
+#pragma once
+
+#include "CEClonerEffectorShared.h"
+#include "CEClonerExtensionBase.h"
+#include "CEPropertyChangeDispatcher.h"
+#include "CEClonerEmitterSpawnExtension.generated.h"
+
+class UNiagaraDataInterfaceCurve;
+
+/** Extension dealing with clones spawning options */
+UCLASS(MinimalAPI, BlueprintType, Within=CEClonerComponent)
+class UCEClonerEmitterSpawnExtension : public UCEClonerExtensionBase
+{
+	GENERATED_BODY()
+
+public:
+	UCEClonerEmitterSpawnExtension()
+		: UCEClonerExtensionBase(
+			TEXT("Spawn")
+			, 0
+#if WITH_EDITOR
+			, TEXT("Emitter")
+			, 4
+#endif
+		)
+	{}
+
+	UFUNCTION(BlueprintCallable, Category="Cloner")
+	CLONEREFFECTOR_API void SetSpawnLoopMode(ECEClonerSpawnLoopMode InMode);
+
+	UFUNCTION(BlueprintPure, Category="Cloner")
+	ECEClonerSpawnLoopMode GetSpawnLoopMode() const
+	{
+		return SpawnLoopMode;
+	}
+
+	UFUNCTION(BlueprintCallable, Category="Cloner")
+	CLONEREFFECTOR_API void SetSpawnLoopIterations(int32 InIterations);
+
+	UFUNCTION(BlueprintPure, Category="Cloner")
+	int32 GetSpawnLoopIterations() const
+	{
+		return SpawnLoopIterations;
+	}
+
+	UFUNCTION(BlueprintCallable, Category="Cloner")
+	CLONEREFFECTOR_API void SetSpawnLoopInterval(float InInterval);
+
+	UFUNCTION(BlueprintPure, Category="Cloner")
+	float GetSpawnLoopInterval() const
+	{
+		return SpawnLoopInterval;
+	}
+
+	UFUNCTION(BlueprintCallable, Category="Cloner")
+	CLONEREFFECTOR_API void SetSpawnBehaviorMode(ECEClonerSpawnBehaviorMode InMode);
+
+	UFUNCTION(BlueprintPure, Category="Cloner")
+	ECEClonerSpawnBehaviorMode GetSpawnBehaviorMode() const
+	{
+		return SpawnBehaviorMode;
+	}
+
+	UFUNCTION(BlueprintCallable, Category="Cloner")
+	CLONEREFFECTOR_API void SetSpawnRate(float InRate);
+
+	UFUNCTION(BlueprintPure, Category="Cloner")
+	float GetSpawnRate() const
+	{
+		return SpawnRate;
+	}
+
+protected:
+	//~ Begin UObject
+#if WITH_EDITOR
+	virtual void PostEditChangeProperty(FPropertyChangedEvent& InPropertyChangedEvent) override;
+#endif
+	//~ End UObject
+
+	//~ Begin UCEClonerExtensionBase
+	virtual void OnExtensionParametersChanged(UCEClonerComponent* InComponent) override;
+	//~ End UCEClonerExtensionBase
+
+	/** How many times do we spawn clones */
+	UPROPERTY(EditInstanceOnly, BlueprintReadWrite, Setter="SetSpawnLoopMode", Getter="GetSpawnLoopMode", DisplayName="LoopMode", Category="Spawn")
+	ECEClonerSpawnLoopMode SpawnLoopMode = ECEClonerSpawnLoopMode::Once;
+
+	/** Amount of spawn iterations for clones */
+	UPROPERTY(EditInstanceOnly, BlueprintReadWrite, Setter="SetSpawnLoopIterations", Getter="GetSpawnLoopIterations", DisplayName="LoopIterations", Category="Spawn", meta=(ClampMin="1", EditCondition="SpawnLoopMode == ECEClonerSpawnLoopMode::Multiple", EditConditionHides))
+	int32 SpawnLoopIterations = 1;
+
+	/** Interval/Duration of spawn for clones */
+	UPROPERTY(EditInstanceOnly, BlueprintReadWrite, Setter="SetSpawnLoopInterval", Getter="GetSpawnLoopInterval", DisplayName="LoopInterval", Category="Spawn", meta=(ClampMin="0"))
+	float SpawnLoopInterval = 1.f;
+
+	/** How does spawn occurs */
+	UPROPERTY(EditInstanceOnly, BlueprintReadWrite, Setter="SetSpawnBehaviorMode", Getter="GetSpawnBehaviorMode", DisplayName="BehaviorMode", Category="Spawn")
+	ECEClonerSpawnBehaviorMode SpawnBehaviorMode = ECEClonerSpawnBehaviorMode::Instant;
+
+	/** How many clones to spawn each seconds */
+	UPROPERTY(EditInstanceOnly, BlueprintReadWrite, Setter="SetSpawnRate", Getter="GetSpawnRate", DisplayName="Rate", Category="Spawn", meta=(ClampMin="0", EditCondition="SpawnBehaviorMode == ECEClonerSpawnBehaviorMode::Rate", EditConditionHides))
+	float SpawnRate = 1.f;
+
+private:
+#if WITH_EDITOR
+	/** Used for PECP */
+	static const TCEPropertyChangeDispatcher<UCEClonerEmitterSpawnExtension> PropertyChangeDispatcher;
+#endif
+};
\ No newline at end of file
diff --git a/Engine/Plugins/Experimental/ClonerEffector/Source/ClonerEffector/Public/Cloner/Extensions/CEClonerRangeExtension.h b/Engine/Plugins/Experimental/ClonerEffector/Source/ClonerEffector/Public/Cloner/Extensions/CEClonerRangeExtension.h
new file mode 100644
--- /dev/null
+++ b/Engine/Plugins/Experimental/ClonerEffector/Source/ClonerEffector/Public/Cloner/Extensions/CEClonerRangeExtension.h
@@ -0,0 +1,164 @@
+// Copyright Epic Games, Inc. All Rights Reserved.
+
+#pragma once
+
+#include "CEClonerExtensionBase.h"
+#include "CEPropertyChangeDispatcher.h"
+#include "CEClonerRangeExtension.generated.h"
+
+/** Extension dealing with range options */
+UCLASS(MinimalAPI, BlueprintType, Within=CEClonerComponent)
+class UCEClonerRangeExtension : public UCEClonerExtensionBase
+{
+	GENERATED_BODY()
+
+public:
+	UCEClonerRangeExtension()
+		: UCEClonerExtensionBase(
+			TEXT("Range")
+			, 0
+#if WITH_EDITOR
+			, TEXT("Range")
+			, 5
+#endif
+		)
+	{}
+
+	UFUNCTION(BlueprintCallable, Category="Cloner")
+	CLONEREFFECTOR_API void SetRangeEnabled(bool bInRangeEnabled);
+
+	UFUNCTION(BlueprintPure, Category="Cloner")
+	bool GetRangeEnabled() const
+	{
+		return bRangeEnabled;
+	}
+
+	UFUNCTION(BlueprintCallable, Category="Cloner")
+	CLONEREFFECTOR_API void SetRangeOffsetMin(const FVector& InRangeOffsetMin);
+
+	UFUNCTION(BlueprintPure, Category="Cloner")
+	const FVector& GetRangeOffsetMin() const
+	{
+		return RangeOffsetMin;
+	}
+
+	UFUNCTION(BlueprintCallable, Category="Cloner")
+	CLONEREFFECTOR_API void SetRangeOffsetMax(const FVector& InRangeOffsetMax);
+
+	UFUNCTION(BlueprintPure, Category="Cloner")
+	const FVector& GetRangeOffsetMax() const
+	{
+		return RangeOffsetMax;
+	}
+
+	UFUNCTION(BlueprintCallable, Category="Cloner")
+	CLONEREFFECTOR_API void SetRangeRotationMin(const FRotator& InRangeRotationMin);
+
+	UFUNCTION(BlueprintPure, Category="Cloner")
+	const FRotator& GetRangeRotationMin() const
+	{
+		return RangeRotationMin;
+	}
+
+	UFUNCTION(BlueprintCallable, Category="Cloner")
+	CLONEREFFECTOR_API void SetRangeRotationMax(const FRotator& InRangeRotationMax);
+
+	UFUNCTION(BlueprintPure, Category="Cloner")
+	const FRotator& GetRangeRotationMax() const
+	{
+		return RangeRotationMax;
+	}
+
+	UFUNCTION(BlueprintCallable, Category="Cloner")
+	CLONEREFFECTOR_API void SetRangeScaleUniform(bool bInRangeScaleUniform);
+
+	UFUNCTION(BlueprintPure, Category="Cloner")
+	bool GetRangeScaleUniform() const
+	{
+		return bRangeScaleUniform;
+	}
+
+	UFUNCTION(BlueprintCallable, Category="Cloner")
+	CLONEREFFECTOR_API void SetRangeScaleMin(const FVector& InRangeScaleMin);
+
+	UFUNCTION(BlueprintPure, Category="Cloner")
+	const FVector& GetRangeScaleMin() const
+	{
+		return RangeScaleMin;
+	}
+
+	UFUNCTION(BlueprintCallable, Category="Cloner")
+	CLONEREFFECTOR_API void SetRangeScaleMax(const FVector& InRangeScaleMax);
+
+	UFUNCTION(BlueprintPure, Category="Cloner")
+	const FVector& GetRangeScaleMax() const
+	{
+		return RangeScaleMax;
+	}
+
+	UFUNCTION(BlueprintCallable, Category="Cloner")
+	CLONEREFFECTOR_API void SetRangeScaleUniformMin(float InRangeScaleUniformMin);
+
+	UFUNCTION(BlueprintPure, Category="Cloner")
+	float GetRangeScaleUniformMin() const
+	{
+		return RangeScaleUniformMin;
+	}
+
+	UFUNCTION(BlueprintCallable, Category="Cloner")
+	CLONEREFFECTOR_API void SetRangeScaleUniformMax(float InRangeScaleUniformMax);
+
+	UFUNCTION(BlueprintPure, Category="Cloner")
+	float GetRangeScaleUniformMax() const
+	{
+		return RangeScaleUniformMax;
+	}
+
+protected:
+	//~ Begin UObject
+#if WITH_EDITOR
+	virtual void PostEditChangeProperty(FPropertyChangedEvent& InPropertyChangedEvent) override;
+#endif
+	//~ End UObject
+
+	//~ Begin UCEClonerExtensionBase
+	virtual void OnExtensionParametersChanged(UCEClonerComponent* InComponent) override;
+	//~ End UCEClonerExtensionBase
+
+	/** Use random range transforms for each clones */
+	UPROPERTY(EditInstanceOnly, BlueprintReadWrite, Setter="SetRangeEnabled", Getter="GetRangeEnabled", DisplayName="Enabled", Category="Range")
+	bool bRangeEnabled = false;
+
+	UPROPERTY(EditInstanceOnly, BlueprintReadWrite, Setter="SetRangeOffsetMin", Getter="GetRangeOffsetMin", DisplayName="OffsetMin", Category="Range", meta=(EditCondition="bRangeEnabled", EditConditionHides))
+	FVector RangeOffsetMin = FVector::ZeroVector;
+
+	UPROPERTY(EditInstanceOnly, BlueprintReadWrite, Setter="SetRangeOffsetMax", Getter="GetRangeOffsetMax", DisplayName="OffsetMax", Category="Range", meta=(EditCondition="bRangeEnabled", EditConditionHides))
+	FVector RangeOffsetMax = FVector::ZeroVector;
+
+	UPROPERTY(EditInstanceOnly, BlueprintReadWrite, Setter="SetRangeRotationMin", Getter="GetRangeRotationMin", DisplayName="RotationMin", Category="Range", meta=(EditCondition="bRangeEnabled", EditConditionHides))
+	FRotator RangeRotationMin = FRotator::ZeroRotator;
+
+	UPROPERTY(EditInstanceOnly, BlueprintReadWrite, Setter="SetRangeRotationMax", Getter="GetRangeRotationMax", DisplayName="RotationMax", Category="Range", meta=(EditCondition="bRangeEnabled", EditConditionHides))
+	FRotator RangeRotationMax = FRotator::ZeroRotator;
+
+	UPROPERTY(EditInstanceOnly, BlueprintReadWrite, Setter="SetRangeScaleUniform", Getter="GetRangeScaleUniform", DisplayName="ScaleUniformEnabled", Category="Range", meta=(EditCondition="bRangeEnabled", EditConditionHides))
+	bool bRangeScaleUniform = false;
+
+	UPROPERTY(EditInstanceOnly, BlueprintReadWrite, Setter="SetRangeScaleMin", Getter="GetRangeScaleMin", DisplayName="ScaleMin", Category="Range", meta=(AllowPreserveRatio, Delta="0.0001", ClampMin="0", EditCondition="bRangeEnabled && !bRangeScaleUniform", EditConditionHides))
+	FVector RangeScaleMin = FVector::OneVector;
+
+	UPROPERTY(EditInstanceOnly, BlueprintReadWrite, Setter="SetRangeScaleMax", Getter="GetRangeScaleMax", DisplayName="ScaleMax", Category="Range", meta=(AllowPreserveRatio, Delta="0.0001", ClampMin="0", EditCondition="bRangeEnabled && !bRangeScaleUniform", EditConditionHides))
+	FVector RangeScaleMax = FVector::OneVector;
+
+	UPROPERTY(EditInstanceOnly, BlueprintReadWrite, Setter="SetRangeScaleUniformMin", Getter="GetRangeScaleUniformMin", DisplayName="ScaleMin", Category="Range", meta=(Delta="0.0001", ClampMin="0", EditCondition="bRangeEnabled && bRangeScaleUniform", EditConditionHides))
+	float RangeScaleUniformMin = 1.f;
+
+	UPROPERTY(EditInstanceOnly, BlueprintReadWrite, Setter="SetRangeScaleUniformMax", Getter="GetRangeScaleUniformMax", DisplayName="ScaleMax", Category="Range", meta=(Delta="0.0001", ClampMin="0", EditCondition="bRangeEnabled && bRangeScaleUniform", EditConditionHides))
+	float RangeScaleUniformMax = 1.f;
+
+private:
+#if WITH_EDITOR
+	/** Used for PECP */
+	static const TCEPropertyChangeDispatcher<UCEClonerRangeExtension> PropertyChangeDispatcher;
+#endif
+};
diff --git a/Engine/Plugins/Experimental/ClonerEffector/Source/ClonerEffector/Public/Cloner/Extensions/CEClonerStepExtension.h b/Engine/Plugins/Experimental/ClonerEffector/Source/ClonerEffector/Public/Cloner/Extensions/CEClonerStepExtension.h
new file mode 100644
--- /dev/null
+++ b/Engine/Plugins/Experimental/ClonerEffector/Source/ClonerEffector/Public/Cloner/Extensions/CEClonerStepExtension.h
@@ -0,0 +1,95 @@
+// Copyright Epic Games, Inc. All Rights Reserved.
+
+#pragma once
+
+#include "CEClonerExtensionBase.h"
+#include "CEPropertyChangeDispatcher.h"
+#include "CEClonerStepExtension.generated.h"
+
+/** Extension dealing with delta step accumulated options */
+UCLASS(MinimalAPI, BlueprintType, Within=CEClonerComponent)
+class UCEClonerStepExtension : public UCEClonerExtensionBase
+{
+	GENERATED_BODY()
+
+public:
+	UCEClonerStepExtension()
+		: UCEClonerExtensionBase(
+			TEXT("Step")
+			, 0
+#if WITH_EDITOR
+			, TEXT("Step")
+			, 5
+#endif
+		)
+	{}
+
+	UFUNCTION(BlueprintCallable, Category="Cloner")
+	CLONEREFFECTOR_API void SetDeltaStepEnabled(bool bInEnabled);
+
+	UFUNCTION(BlueprintPure, Category="Cloner")
+	bool GetDeltaStepEnabled() const
+	{
+		return bDeltaStepEnabled;
+	}
+
+	UFUNCTION(BlueprintCallable, Category="Cloner")
+	CLONEREFFECTOR_API void SetDeltaStepPosition(const FVector& InPosition);
+
+	UFUNCTION(BlueprintPure, Category="Cloner")
+	FVector GetDeltaStepPosition() const
+	{
+		return DeltaStepPosition;
+	}
+
+	UFUNCTION(BlueprintCallable, Category="Cloner")
+	CLONEREFFECTOR_API void SetDeltaStepRotation(const FRotator& InRotation);
+
+	UFUNCTION(BlueprintPure, Category="Cloner")
+	FRotator GetDeltaStepRotation() const
+	{
+		return DeltaStepRotation;
+	}
+
+	UFUNCTION(BlueprintCallable, Category="Cloner")
+	CLONEREFFECTOR_API void SetDeltaStepScale(const FVector& InScale);
+
+	UFUNCTION(BlueprintPure, Category="Cloner")
+	FVector GetDeltaStepScale() const
+	{
+		return DeltaStepScale;
+	}
+
+protected:
+	//~ Begin UObject
+#if WITH_EDITOR
+	virtual void PostEditChangeProperty(FPropertyChangedEvent& InPropertyChangedEvent) override;
+#endif
+	//~ End UObject
+
+	//~ Begin UCEClonerExtensionBase
+	virtual void OnExtensionParametersChanged(UCEClonerComponent* InComponent) override;
+	//~ End UCEClonerExtensionBase
+
+	/** Enable steps to add delta variation on each clone instance */
+	UPROPERTY(EditInstanceOnly, BlueprintReadWrite, Setter="SetDeltaStepEnabled", Getter="GetDeltaStepEnabled", DisplayName="Enabled", Category="Step")
+	bool bDeltaStepEnabled = false;
+
+	/** Amount of position difference between one step and the next one */
+	UPROPERTY(EditInstanceOnly, BlueprintReadWrite, Setter="SetDeltaStepPosition", Getter="GetDeltaStepPosition", DisplayName="Position", Category="Step", meta=(EditCondition="bDeltaStepEnabled", EditConditionHides))
+	FVector DeltaStepPosition = FVector(0.f);
+
+	/** Amount of rotation difference between one step and the next one */
+	UPROPERTY(EditInstanceOnly, BlueprintReadWrite, Setter="SetDeltaStepRotation", Getter="GetDeltaStepRotation", DisplayName="Rotation", Category="Step", meta=(EditCondition="bDeltaStepEnabled", EditConditionHides))
+	FRotator DeltaStepRotation = FRotator(0.f);
+
+	/** Amount of scale difference between one step and the next one */
+	UPROPERTY(EditInstanceOnly, BlueprintReadWrite, Setter="SetDeltaStepScale", Getter="GetDeltaStepScale", DisplayName="Scale", Category="Step", meta=(AllowPreserveRatio, Delta="0.0001", EditCondition="bDeltaStepEnabled", EditConditionHides))
+	FVector DeltaStepScale = FVector(0.f);
+
+private:
+#if WITH_EDITOR
+	/** Used for PECP */
+	static const TCEPropertyChangeDispatcher<UCEClonerStepExtension> PropertyChangeDispatcher;
+#endif
+};
\ No newline at end of file
diff --git a/Engine/Plugins/Experimental/ClonerEffector/Source/ClonerEffector/Public/Cloner/Layouts/CEClonerLayoutBase.h b/Engine/Plugins/Experimental/ClonerEffector/Source/ClonerEffector/Public/Cloner/Layouts/CEClonerLayoutBase.h
--- a/Engine/Plugins/Experimental/ClonerEffector/Source/ClonerEffector/Public/Cloner/Layouts/CEClonerLayoutBase.h
+++ b/Engine/Plugins/Experimental/ClonerEffector/Source/ClonerEffector/Public/Cloner/Layouts/CEClonerLayoutBase.h
@@ -6,8 +6,8 @@
 #include "UObject/Object.h"
 #include "CEClonerLayoutBase.generated.h"
 
-class ACEClonerActor;
 class UCEClonerComponent;
+class UCEClonerExtensionBase;
 class UNiagaraMeshRendererProperties;
 class UNiagaraSystem;
 
@@ -19,13 +19,13 @@ class UNiagaraSystem;
  * 3. Expose all new system specific parameters in the layout extended class and update them when required
  * Your new layout is ready and will be available in the cloner in the layout dropdown
  */
-UCLASS(MinimalAPI, Abstract, BlueprintType, Within=CEClonerActor, AutoExpandCategories=("Layout"))
+UCLASS(MinimalAPI, Abstract, BlueprintType, AutoExpandCategories=("Layout"))
 class UCEClonerLayoutBase : public UObject
 {
 	GENERATED_BODY()
 
 public:
-	static inline constexpr TCHAR LayoutBaseAssetPath[] = TEXT("/Script/Niagara.NiagaraSystem'/ClonerEffector/Systems/NS_ClonerBase.NS_ClonerBase'");
+	static constexpr TCHAR LayoutBaseAssetPath[] = TEXT("/Script/Niagara.NiagaraSystem'/ClonerEffector/Systems/NS_ClonerBase.NS_ClonerBase'");
 
 	UCEClonerLayoutBase()
 		: UCEClonerLayoutBase(NAME_None, FString())
@@ -62,13 +62,17 @@ public:
 		return DataInterfaces;
 	}
 
-	/** Get the cloner actor using this layout */
-	UFUNCTION(BlueprintPure, Category="Cloner|Layout")
-	CLONEREFFECTOR_API ACEClonerActor* GetClonerActor() const;
+	TMulticastDelegateRegistration<void(UCEClonerLayoutBase*, bool)>& OnLayoutLoadedDelegate()
+	{
+		return OnClonerLayoutLoadedDelegate;
+	}
 
 	/** Get the cloner component using this layout */
+	UFUNCTION(BlueprintPure, Category="Cloner|Layout")
 	UCEClonerComponent* GetClonerComponent() const;
 
+	AActor* GetClonerActor() const;
+
 	/** Request refresh layout next tick */
 	void UpdateLayoutParameters(bool bInUpdateCloner = true, bool bInImmediate = false);
 
@@ -83,7 +87,7 @@ public:
 	CLONEREFFECTOR_API bool IsLayoutLoaded() const;
 
 	/** Load this layout system if not already loaded */
-	bool LoadLayout();
+	void LoadLayout();
 
 	/** Free the loaded system and return to idle state */
 	bool UnloadLayout();
@@ -101,6 +105,9 @@ public:
 	/** Copies this layout data interfaces to other layout */
 	bool CopyTo(UCEClonerLayoutBase* InOtherLayout) const;
 
+	/** Gets the cloner extensions supported by this layout */
+	virtual TSet<FName> GetSupportedExtensions() const;
+
 protected:
 	/** Called once after layout is loaded */
 	virtual void OnLayoutLoaded() {}
@@ -120,7 +127,10 @@ protected:
 	void OnLayoutPropertyChanged();
 
 private:
-	UNiagaraSystem* LoadSystemPath(const FString& InPath) const;
+	DECLARE_MULTICAST_DELEGATE_TwoParams(FOnClonerLayoutLoaded, UCEClonerLayoutBase* /** InLayout */, bool /** bInSuccess */)
+	FOnClonerLayoutLoaded OnClonerLayoutLoadedDelegate;
+
+	void OnSystemPackageLoaded(const FName& InName, UPackage* InPackage, EAsyncLoadingResult::Type InResult);
 
 	/** Layout name to display in layout options */
 	UPROPERTY(Transient)
@@ -141,4 +151,6 @@ private:
 	/** Data interfaces used by the effectors */
 	UPROPERTY(Transient)
 	FCEClonerEffectorDataInterfaces DataInterfaces;
+
+	int32 LoadRequestIdentifier = INDEX_NONE;
 };
\ No newline at end of file
