MotionDesign : ClonerEffector

diff --git a/Engine/Plugins/Experimental/ClonerEffector/Source/ClonerEffector/Public/Effector/CEEffectorActor.h b/Engine/Plugins/Experimental/ClonerEffector/Source/ClonerEffector/Public/Effector/CEEffectorActor.h
--- a/Engine/Plugins/Experimental/ClonerEffector/Source/ClonerEffector/Public/Effector/CEEffectorActor.h
+++ b/Engine/Plugins/Experimental/ClonerEffector/Source/ClonerEffector/Public/Effector/CEEffectorActor.h
@@ -3,15 +3,18 @@
 #pragma once
 
 #include "CEClonerEffectorShared.h"
-#include "CEPropertyChangeDispatcher.h"
 #include "GameFramework/Actor.h"
-#include "PropertyBag.h"
 #include "CEEffectorActor.generated.h"
 
 class ACEClonerActor;
+class UCEEffectorComponent;
 class UDynamicMeshComponent;
 
-UCLASS(MinimalAPI, BlueprintType, HideCategories=(Rendering,Replication,Collision,HLOD,Physics,Networking,Input,Actor,Cooking,LevelInstance,DataLayers), DisplayName = "Motion Design Effector Actor")
+UCLASS(
+	MinimalAPI
+	, BlueprintType
+	, HideCategories=(Rendering,Replication,Collision,HLOD,Physics,Networking,Input,Actor,Cooking,LevelInstance,DataLayers)
+	, DisplayName = "Motion Design Effector Actor")
 class ACEEffectorActor : public AActor
 {
 	GENERATED_BODY()
@@ -26,759 +29,218 @@ public:
 
 	ACEEffectorActor();
 
-	UFUNCTION(BlueprintCallable, Category="Effector")
-	CLONEREFFECTOR_API void SetType(ECEClonerEffectorType InType);
-
-	UFUNCTION(BlueprintPure, Category="Effector")
-	ECEClonerEffectorType GetType() const
-	{
-		return Type;
-	}
-
-	UFUNCTION(BlueprintCallable, Category="Effector")
-	CLONEREFFECTOR_API void SetInvertType(bool bInInvert);
-
-	UFUNCTION(BlueprintPure, Category="Effector")
-	bool GetInvertType() const
-	{
-		return bInvertType;
-	}
-
-	UFUNCTION(BlueprintCallable, Category="Effector")
-	CLONEREFFECTOR_API void SetEasing(ECEClonerEasing InEasing);
-
-	UFUNCTION(BlueprintPure, Category="Effector")
-	ECEClonerEasing GetEasing() const
-	{
-		return Easing;
-	}
-
-	UFUNCTION(BlueprintCallable, Category="Effector")
-	CLONEREFFECTOR_API void SetMagnitude(float InMagnitude);
-
-	UFUNCTION(BlueprintPure, Category="Effector")
-	float GetMagnitude() const
-	{
-		return Magnitude;
-	}
-
-	UFUNCTION(BlueprintCallable, Category="Effector")
-	CLONEREFFECTOR_API void SetOuterRadius(float InRadius);
-
-	UFUNCTION(BlueprintPure, Category="Effector")
-	float GetOuterRadius() const
-	{
-		return OuterRadius;
-	}
-
-	UFUNCTION(BlueprintCallable, Category="Effector")
-	CLONEREFFECTOR_API void SetInnerRadius(float InRadius);
-
-	UFUNCTION(BlueprintPure, Category="Effector")
-	float GetInnerRadius() const
-	{
-		return InnerRadius;
-	}
-
-	UFUNCTION(BlueprintCallable, Category="Effector")
-	CLONEREFFECTOR_API void SetInnerExtent(const FVector& InExtent);
-
-	UFUNCTION(BlueprintPure, Category="Effector")
-	FVector GetInnerExtent() const
-	{
-		return InnerExtent;
-	}
-
-	UFUNCTION(BlueprintCallable, Category="Effector")
-	CLONEREFFECTOR_API void SetOuterExtent(const FVector& InExtent);
-
-	UFUNCTION(BlueprintPure, Category="Effector")
-	FVector GetOuterExtent() const
-	{
-		return OuterExtent;
-	}
-
-	UFUNCTION(BlueprintCallable, Category="Effector")
-	CLONEREFFECTOR_API void SetPlaneSpacing(float InSpacing);
-
-	UFUNCTION(BlueprintPure, Category="Effector")
-	float GetPlaneSpacing() const
-	{
-		return PlaneSpacing;
-	}
-
-	UFUNCTION(BlueprintCallable, Category="Effector")
-	CLONEREFFECTOR_API void SetRadialAngle(float InAngle);
-
-	UFUNCTION(BlueprintPure, Category="Effector")
-	float GetRadialAngle() const
-	{
-		return RadialAngle;
-	}
-
-	UFUNCTION(BlueprintCallable, Category="Effector")
-	CLONEREFFECTOR_API void SetRadialMinRadius(float InRadius);
-
-	UFUNCTION(BlueprintPure, Category="Effector")
-	float GetRadialMinRadius() const
-	{
-		return RadialMinRadius;
-	}
-
-	UFUNCTION(BlueprintCallable, Category="Effector")
-	CLONEREFFECTOR_API void SetRadialMaxRadius(float InRadius);
-
-	UFUNCTION(BlueprintPure, Category="Effector")
-	float GetRadialMaxRadius() const
-	{
-		return RadialMaxRadius;
-	}
-
-	UFUNCTION(BlueprintCallable, Category="Effector")
-	CLONEREFFECTOR_API void SetTorusRadius(float InRadius);
-
-	UFUNCTION(BlueprintPure, Category="Effector")
-	float GetTorusRadius() const
-	{
-		return TorusRadius;
-	}
-
-	UFUNCTION(BlueprintCallable, Category="Effector")
-	CLONEREFFECTOR_API void SetTorusInnerRadius(float InRadius);
-
-	UFUNCTION(BlueprintPure, Category="Effector")
-	float GetTorusInnerRadius() const
-	{
-		return TorusInnerRadius;
-	}
-
-	UFUNCTION(BlueprintCallable, Category="Effector")
-	CLONEREFFECTOR_API void SetTorusOuterRadius(float InRadius);
-
-	UFUNCTION(BlueprintPure, Category="Effector")
-	float GetTorusOuterRadius() const
-	{
-		return TorusOuterRadius;
-	}
-
-	UFUNCTION(BlueprintCallable, Category="Effector")
-	CLONEREFFECTOR_API void SetOffset(const FVector& InOffset);
-
-	UFUNCTION(BlueprintPure, Category="Effector")
-	FVector GetOffset() const
-	{
-		return Offset;
-	}
-
-	UFUNCTION(BlueprintCallable, Category="Effector")
-	CLONEREFFECTOR_API void SetRotation(const FRotator& InRotation);
-
-	UFUNCTION(BlueprintPure, Category="Effector")
-	FRotator GetRotation() const
-	{
-		return Rotation;
-	}
-
-	UFUNCTION(BlueprintCallable, Category="Effector")
-	CLONEREFFECTOR_API void SetScale(const FVector& InScale);
-
-	UFUNCTION(BlueprintPure, Category="Effector")
-	FVector GetScale() const
-	{
-		return Scale;
-	}
-
-	UFUNCTION(BlueprintCallable, Category="Effector")
-	CLONEREFFECTOR_API void SetEnabled(bool bInEnable);
-
-	UFUNCTION(BlueprintPure, Category="Effector")
-	bool GetEnabled() const
-	{
-		return bEnabled;
-	}
-
-	UFUNCTION(BlueprintCallable, Category="Effector")
-	CLONEREFFECTOR_API void SetMode(ECEClonerEffectorMode InMode);
-
-	UFUNCTION(BlueprintPure, Category="Effector")
-	ECEClonerEffectorMode GetMode() const
-	{
-		return Mode;
-	}
-
-	UFUNCTION(BlueprintCallable, Category="Effector")
-	CLONEREFFECTOR_API void SetTargetActor(AActor* InTargetActor);
-
-	UFUNCTION(BlueprintPure, Category="Effector")
-	AActor* GetTargetActor() const
-	{
-		return TargetActorWeak.Get();
-	}
-
-	void SetTargetActorWeak(const TWeakObjectPtr<AActor>& InTargetActor);
-
-	TWeakObjectPtr<AActor> GetTargetActorWeak() const
-	{
-		return TargetActorWeak;
-	}
-
-	UFUNCTION(BlueprintCallable, Category="Effector")
-	CLONEREFFECTOR_API void SetOrientationForceEnabled(bool bInForceEnabled);
-
-	UFUNCTION(BlueprintPure, Category="Effector")
-	bool GetOrientationForceEnabled() const
+	UCEEffectorComponent* GetEffectorComponent() const
 	{
-		return bOrientationForceEnabled;
+		return EffectorComponent;
 	}
 
-	UFUNCTION(BlueprintCallable, Category="Effector")
-	CLONEREFFECTOR_API void SetOrientationForceRate(float InForceOrientationRate);
-
-	UFUNCTION(BlueprintPure, Category="Effector")
-	float GetOrientationForceRate() const
-	{
-		return OrientationForceRate;
-	}
-
-	UFUNCTION(BlueprintCallable, Category="Effector")
-	CLONEREFFECTOR_API void SetOrientationForceMin(const FVector& InForceOrientationMin);
-
-	UFUNCTION(BlueprintPure, Category="Effector")
-	const FVector& GetOrientationForceMin() const
-	{
-		return OrientationForceMin;
-	}
-
-	UFUNCTION(BlueprintCallable, Category="Effector")
-	CLONEREFFECTOR_API void SetOrientationForceMax(const FVector& InForceOrientationMax);
-
-	UFUNCTION(BlueprintPure, Category="Effector")
-	const FVector& GetOrientationForceMax() const
-	{
-		return OrientationForceMax;
-	}
-
-	UFUNCTION(BlueprintCallable, Category="Effector")
-	CLONEREFFECTOR_API void SetVortexForceEnabled(bool bInForceEnabled);
-
-	UFUNCTION(BlueprintPure, Category="Effector")
-	bool GetVortexForceEnabled() const
-	{
-		return bVortexForceEnabled;
-	}
-
-	UFUNCTION(BlueprintCallable, Category="Effector")
-	CLONEREFFECTOR_API void SetVortexForceAmount(float InForceVortexAmount);
-
-	UFUNCTION(BlueprintPure, Category="Effector")
-	float GetVortexForceAmount() const
-	{
-		return VortexForceAmount;
-	}
-
-	UFUNCTION(BlueprintCallable, Category="Effector")
-	CLONEREFFECTOR_API void SetVortexForceAxis(const FVector& InForceVortexAxis);
-
-	UFUNCTION(BlueprintPure, Category="Effector")
-	const FVector& GetVortexForceAxis() const
-	{
-		return VortexForceAxis;
-	}
-
-	UFUNCTION(BlueprintCallable, Category="Effector")
-	CLONEREFFECTOR_API void SetCurlNoiseForceEnabled(bool bInForceEnabled);
-
-	UFUNCTION(BlueprintPure, Category="Effector")
-	bool GetCurlNoiseForceEnabled() const
-	{
-		return bCurlNoiseForceEnabled;
-	}
-
-	UFUNCTION(BlueprintCallable, Category="Effector")
-	CLONEREFFECTOR_API void SetCurlNoiseForceStrength(float InForceCurlNoiseStrength);
-
-	UFUNCTION(BlueprintPure, Category="Effector")
-	float GetCurlNoiseForceStrength() const
-	{
-		return CurlNoiseForceStrength;
-	}
-
-	UFUNCTION(BlueprintCallable, Category="Effector")
-	CLONEREFFECTOR_API void SetCurlNoiseForceFrequency(float InForceCurlNoiseFrequency);
-
-	UFUNCTION(BlueprintPure, Category="Effector")
-	float GetCurlNoiseForceFrequency() const
-	{
-		return CurlNoiseForceFrequency;
-	}
-
-	UFUNCTION(BlueprintCallable, Category="Effector")
-	CLONEREFFECTOR_API void SetAttractionForceEnabled(bool bInForceEnabled);
-
-	UFUNCTION(BlueprintPure, Category="Effector")
-	bool GetAttractionForceEnabled() const
-	{
-		return bAttractionForceEnabled;
-	}
-
-	UFUNCTION(BlueprintCallable, Category="Effector")
-	CLONEREFFECTOR_API void SetAttractionForceStrength(float InForceStrength);
-
-	UFUNCTION(BlueprintPure, Category="Effector")
-	float GetAttractionForceStrength() const
-	{
-		return AttractionForceStrength;
-	}
-
-	UFUNCTION(BlueprintCallable, Category="Effector")
-	CLONEREFFECTOR_API void SetAttractionForceFalloff(float InForceFalloff);
-
-	UFUNCTION(BlueprintPure, Category="Effector")
-	float GetAttractionForceFalloff() const
-	{
-		return AttractionForceFalloff;
-	}
-
-	UFUNCTION(BlueprintCallable, Category="Effector")
-	CLONEREFFECTOR_API void SetGravityForceEnabled(bool bInForceEnabled);
-
-	UFUNCTION(BlueprintPure, Category="Effector")
-	bool GetGravityForceEnabled() const
-	{
-		return bGravityForceEnabled;
-	}
-
-	UFUNCTION(BlueprintCallable, Category="Effector")
-	CLONEREFFECTOR_API void SetGravityForceAcceleration(const FVector& InAcceleration);
-
-	UFUNCTION(BlueprintPure, Category="Effector")
-	FVector GetGravityForceAcceleration() const
-	{
-		return GravityForceAcceleration;
-	}
-
-	UFUNCTION(BlueprintCallable, Category="Effector")
-	CLONEREFFECTOR_API void SetLocationStrength(const FVector& InStrength);
-
-	UFUNCTION(BlueprintPure, Category="Effector")
-	FVector GetLocationStrength() const
-	{
-		return LocationStrength;
-	}
-
-	UFUNCTION(BlueprintCallable, Category="Effector")
-	CLONEREFFECTOR_API void SetRotationStrength(const FRotator& InStrength);
-
-	UFUNCTION(BlueprintPure, Category="Effector")
-	FRotator GetRotationStrength() const
-	{
-		return RotationStrength;
-	}
-
-	UFUNCTION(BlueprintCallable, Category="Effector")
-	CLONEREFFECTOR_API void SetScaleStrength(const FVector& InStrength);
-
-	UFUNCTION(BlueprintPure, Category="Effector")
-	FVector GetScaleStrength() const
-	{
-		return ScaleStrength;
-	}
-
-	UFUNCTION(BlueprintCallable, Category="Effector")
-	CLONEREFFECTOR_API void SetPan(const FVector& InPan);
-
-	UFUNCTION(BlueprintPure, Category="Effector")
-	FVector GetPan() const
-	{
-		return Pan;
-	}
-
-	UFUNCTION(BlueprintCallable, Category="Effector")
-	CLONEREFFECTOR_API void SetFrequency(float InFrequency);
-
-	UFUNCTION(BlueprintPure, Category="Effector")
-	float GetFrequency() const
-	{
-		return Frequency;
-	}
-
-	UFUNCTION(BlueprintCallable, Category="Effector")
-	CLONEREFFECTOR_API void SetColor(const FLinearColor& InColor);
-
-	UFUNCTION(BlueprintPure, Category="Effector")
-	const FLinearColor& GetColor() const
-	{
-		return Color;
-	}
-
-	UFUNCTION(BlueprintCallable, Category="Effector")
-	CLONEREFFECTOR_API void SetPushDirection(ECEClonerEffectorPushDirection InDirection);
-
-	UFUNCTION(BlueprintPure, Category="Effector")
-	ECEClonerEffectorPushDirection GetPushDirection() const
-	{
-		return PushDirection;
-	}
-
-	UFUNCTION(BlueprintCallable, Category="Effector")
-	CLONEREFFECTOR_API void SetPushStrength(const FVector& InStrength);
-
-	UFUNCTION(BlueprintPure, Category="Effector")
-	const FVector& GetPushStrength() const
-	{
-		return PushStrength;
-	}
-
-#if WITH_EDITOR
-	UFUNCTION(BlueprintCallable, Category="Effector")
-	CLONEREFFECTOR_API void SetVisualizerComponentVisible(bool bInVisible);
-
-	UFUNCTION(BlueprintPure, Category="Effector")
-	bool GetVisualizerComponentVisible() const
-	{
-		return bVisualizerComponentVisible;
-	}
-
-	UFUNCTION(BlueprintCallable, Category="Effector")
-	CLONEREFFECTOR_API void SetVisualizerSpriteVisible(bool bInVisible);
-
-	UFUNCTION(BlueprintPure, Category="Effector")
-	bool GetVisualizerSpriteVisible() const
-	{
-		return bVisualizerSpriteVisible;
-	}
-#endif
-
-protected:
-	static constexpr int32 InnerVisualizerId = 0;
-	static constexpr int32 OuterVisualizerId = 1;
-	static constexpr TCHAR VisualizerColorName[] = TEXT("VisualizerColor");
-
 	//~ Begin AActor
-	virtual void Destroyed() override;
-	virtual void PostActorCreated() override;
 #if WITH_EDITOR
 	virtual FString GetDefaultActorLabel() const override;
 #endif
 	//~ End AActor
 
+protected:
 	//~ Begin UObject
-#if WITH_EDITOR
-	virtual void PostEditChangeProperty(FPropertyChangedEvent& InPropertyChangedEvent) override;
-	virtual void PostEditImport() override;
-	virtual void PostEditUndo() override;
-#endif
+	virtual void Serialize(FArchive& InArchive) override;
 	virtual void PostLoad() override;
-	virtual void PostDuplicate(EDuplicateMode::Type InDuplicateMode) override;
 	//~ End UObject
 
-	void RegisterToChannel();
-	int32 GetChannelIdentifier() const;
-	void OnEffectorSubsystemInitialized(const UWorld* InWorld);
-
-	void OnClonerLinked(ACEClonerActor* InCloner);
-	void OnClonerUnlinked(ACEClonerActor* InCloner);
-
-	FCEClonerEffectorChannelData& GetChannelData();
-
-	void OnEffectorTransformed(USceneComponent* InUpdatedComponent, EUpdateTransformFlags InUpdateTransformFlags, ETeleportType InTeleport);
-
-	/** Will update all values of this effector on all clones */
-	void OnEffectorChanged();
-
-	/** Called when effector is enabled/disabled */
-	void OnEnabledChanged();
-
-	/** Called when we disable this effector */
-	void OnEffectorDisabled();
-
-	/** Update mode for this effector */
-	void OnModeChanged();
-
-	/** Update shape type for this effector */
-	void OnTypeChanged();
-
-	/** Update values for box type effector */
-	void OnBoxChanged();
-
-	/** Update values for sphere type effector */
-	void OnSphereChanged();
-
-	/** Update values for plane type effector */
-	void OnPlaneChanged();
-
-	/** Update values for radial type effector */
-	void OnRadialChanged();
-
-	/** Update values for torus type effector */
-	void OnTorusChanged();
-
-	/** Update values for unbound type effector */
-	void OnUnboundChanged();
-
-	/** Called when transform of effector has changed */
-	void OnTransformChanged();
-
-	/** Update values for easing */
-	void OnEasingChanged();
-
-	/** Update values for magnitude */
-	void OnMagnitudeChanged();
-
-	/** Called when default mode options are changed */
-	void OnTransformOptionsChanged();
-
-	/** Update actor targeted by this effector on cloners */
-	void OnTargetActorChanged();
-
-	/** Called when target mode options are changed */
-	void OnTargetOptionsChanged();
-
-	/** Called when noise field mode options are changed */
-	void OnNoiseFieldOptionsChanged();
-
-	void OnPushOptionsChanged();
-
-	/** Called when force options are changed */
-	void OnForceOptionsChanged();
-
-	/** Called when force enabled state changed */
-	void OnForceEnabledChanged();
-
-	/** Update particle color affected by this effector */
-    void OnColorChanged();
-
-#if WITH_EDITOR
-	/** Update visualizers options of this effector */
-	void OnVisualizerOptionsChanged();
-
-	/** Called when developer settings are changed */
-	void OnEffectorDeveloperSettingsChanged(UObject* InSettings, FPropertyChangedEvent& InEvent);
+private:
+	void MigrateDeprecatedProperties();
 
-	/** Update the mesh of a component visualizer */
-	void UpdateVisualizer(int32 InVisualizerId, TFunction<void(UDynamicMesh*)> InMeshFunction) const;
-#endif
+	UPROPERTY(VisibleInstanceOnly, Category="Effector")
+	TObjectPtr<UCEEffectorComponent> EffectorComponent;
 
-	/** Update all types options */
-	void UpdateEffectorTypes();
+	int32 MigrateToVersion = INDEX_NONE;
 
-	/** Is this effector enabled/disabled on linked cloners */
-	UPROPERTY(EditAnywhere, BlueprintReadWrite, Setter="SetEnabled", Getter="GetEnabled", Category="Effector")
+	UE_DEPRECATED(5.5, "Moved to component and extensions")
+	UPROPERTY()
 	bool bEnabled = true;
 
-	/** The ratio effect of the effector on clones */
-	UPROPERTY(EditAnywhere, BlueprintReadWrite, Setter, Getter, Category="Effector", meta=(ClampMin="0", ClampMax="1"))
+	UE_DEPRECATED(5.5, "Moved to component and extensions")
+	UPROPERTY()
 	float Magnitude = 1.f;
 
-	/** Affected clones color passed over to material */
-	UPROPERTY(EditAnywhere, BlueprintReadWrite, Setter, Getter, Category="Effector")
+	UE_DEPRECATED(5.5, "Moved to component and extensions")
+	UPROPERTY()
 	FLinearColor Color = FLinearColor::Red;
 
-	/** Type of effector to apply on cloners instances */
-	UPROPERTY(EditAnywhere, BlueprintReadWrite, Setter, Getter, Category="Shape")
+	UE_DEPRECATED(5.5, "Moved to component and extensions")
+	UPROPERTY()
 	ECEClonerEffectorType Type = ECEClonerEffectorType::Sphere;
 
-	/** Invert the type effect, instead of affecting the inside of a zone, will affect the outside */
-	UPROPERTY(EditAnywhere, BlueprintReadWrite, Setter="SetInvertType", Getter="GetInvertType", Category="Shape")
+	UE_DEPRECATED(5.5, "Moved to component and extensions")
+	UPROPERTY()
 	bool bInvertType = false;
 
-	/** Weight easing function applied to lerp transforms */
-	UPROPERTY(EditAnywhere, BlueprintReadWrite, Setter, Getter, Category="Shape", meta=(EditCondition="Type != ECEClonerEffectorType::Unbound", EditConditionHides))
+	UE_DEPRECATED(5.5, "Moved to component and extensions")
+	UPROPERTY()
 	ECEClonerEasing Easing = ECEClonerEasing::Linear;
 
-	/** Inner radius of sphere, all clones inside will be affected with a maximum weight */
-	UPROPERTY(EditAnywhere, BlueprintReadWrite, Setter, Getter, Category="Shape", meta=(ClampMin="0", EditCondition="Type == ECEClonerEffectorType::Sphere", EditConditionHides))
+	UE_DEPRECATED(5.5, "Moved to component and extensions")
+	UPROPERTY()
 	float InnerRadius = 50.f;
 
-	/** Outer radius of sphere, all clones outside will not be affected */
-	UPROPERTY(EditAnywhere, BlueprintReadWrite, Setter, Getter, Category="Shape", meta=(ClampMin="0", EditCondition="Type == ECEClonerEffectorType::Sphere", EditConditionHides))
+	UE_DEPRECATED(5.5, "Moved to component and extensions")
+	UPROPERTY()
 	float OuterRadius = 200.f;
 
-	/** Inner extent of box, all clones inside will be affected with a maximum weight */
-	UPROPERTY(EditAnywhere, BlueprintReadWrite, Setter, Getter, Category="Shape", meta=(ClampMin="0", AllowPreserveRatio, EditCondition="Type == ECEClonerEffectorType::Box", EditConditionHides))
+	UE_DEPRECATED(5.5, "Moved to component and extensions")
+	UPROPERTY()
 	FVector InnerExtent = FVector(50.f);
 
-	/** Outer extent of box, all clones outside will not be affected */
-	UPROPERTY(EditAnywhere, BlueprintReadWrite, Setter, Getter, Category="Shape", meta=(ClampMin="0", AllowPreserveRatio, EditCondition="Type == ECEClonerEffectorType::Box", EditConditionHides))
+	UE_DEPRECATED(5.5, "Moved to component and extensions")
+	UPROPERTY()
 	FVector OuterExtent = FVector(200.f);
 
-	/** Plane spacing, everything inside this zone will be affected */
-	UPROPERTY(EditAnywhere, BlueprintReadWrite, Setter, Getter, Category="Shape", meta=(ClampMin="0", EditCondition="Type == ECEClonerEffectorType::Plane", EditConditionHides))
+	UE_DEPRECATED(5.5, "Moved to component and extensions")
+	UPROPERTY()
 	float PlaneSpacing = 200.f;
 
-	/** Radial angle in degree, everything within the angle will be affected */
-	UPROPERTY(EditAnywhere, BlueprintReadWrite, Setter, Getter, Category="Shape", meta=(ClampMin="0", ClampMax="360", EditCondition="Type == ECEClonerEffectorType::Radial", EditConditionHides))
+	UE_DEPRECATED(5.5, "Moved to component and extensions")
+	UPROPERTY()
 	float RadialAngle = 180.f;
 
-	/** Minimum radius for the radial effect to be applied on clones, below clones will not be affected */
-	UPROPERTY(EditAnywhere, BlueprintReadWrite, Setter, Getter, Category="Shape", meta=(ClampMin="0", EditCondition="Type == ECEClonerEffectorType::Radial", EditConditionHides))
+	UE_DEPRECATED(5.5, "Moved to component and extensions")
+	UPROPERTY()
 	float RadialMinRadius = 0.f;
 
-	/** Maximum radius for the radial effect to be applied on clones, above clones will not be affected */
-	UPROPERTY(EditAnywhere, BlueprintReadWrite, Setter, Getter, Category="Shape", meta=(ClampMin="0", EditCondition="Type == ECEClonerEffectorType::Radial", EditConditionHides))
+	UE_DEPRECATED(5.5, "Moved to component and extensions")
+	UPROPERTY()
 	float RadialMaxRadius = 1000.f;
 
-	/** Main torus radius from center to the edge where inner and outer tube will be revolved */
-	UPROPERTY(EditAnywhere, BlueprintReadWrite, Setter, Getter, Category="Shape", meta=(ClampMin="0", EditCondition="Type == ECEClonerEffectorType::Torus", EditConditionHides))
+	UE_DEPRECATED(5.5, "Moved to component and extensions")
+	UPROPERTY()
 	float TorusRadius = 250.f;
 
-	/** Minimum revolved radius for the torus effect, clones contained inside will be affected with a maximum weight */
-	UPROPERTY(EditAnywhere, BlueprintReadWrite, Setter, Getter, Category="Shape", meta=(ClampMin="0", EditCondition="Type == ECEClonerEffectorType::Torus", EditConditionHides))
+	UE_DEPRECATED(5.5, "Moved to component and extensions")
+	UPROPERTY()
 	float TorusInnerRadius = 50.f;
 
-	/** Maximum revolved radius for the torus effect, clones outside of it will not be affected */
-	UPROPERTY(EditAnywhere, BlueprintReadWrite, Setter, Getter, Category="Shape", meta=(ClampMin="0", EditCondition="Type == ECEClonerEffectorType::Torus", EditConditionHides))
+	UE_DEPRECATED(5.5, "Moved to component and extensions")
+	UPROPERTY()
 	float TorusOuterRadius = 200.f;
 
-	/** Mode of effector for each clones instances */
-	UPROPERTY(EditAnywhere, BlueprintReadWrite, Setter, Getter, Category="Mode")
+	UE_DEPRECATED(5.5, "Moved to component and extensions")
+	UPROPERTY()
 	ECEClonerEffectorMode Mode = ECEClonerEffectorMode::Default;
 
-	/** Offset applied on affected clones */
-	UPROPERTY(EditAnywhere, BlueprintReadWrite, Setter, Getter, Category="Mode", meta=(EditCondition="Mode == ECEClonerEffectorMode::Default", EditConditionHides))
+	UE_DEPRECATED(5.5, "Moved to component and extensions")
+	UPROPERTY()
 	FVector Offset = FVector::ZeroVector;
 
-	/** Rotation applied on affected clones */
-	UPROPERTY(EditAnywhere, BlueprintReadWrite, Setter, Getter, Category="Mode", meta=(ClampMin="-180", ClampMax="180", UIMin="-180", UIMax="180", EditCondition="Mode == ECEClonerEffectorMode::Default", EditConditionHides))
+	UE_DEPRECATED(5.5, "Moved to component and extensions")
+	UPROPERTY()
 	FRotator Rotation = FRotator::ZeroRotator;
 
-	/** Scale applied on affected clones */
-	UPROPERTY(EditAnywhere, BlueprintReadWrite, Setter, Getter, Category="Mode", meta=(ClampMin="0", AllowPreserveRatio, Delta="0.01", EditCondition="Mode == ECEClonerEffectorMode::Default", EditConditionHides))
+	UE_DEPRECATED(5.5, "Moved to component and extensions")
+	UPROPERTY()
 	FVector Scale = FVector::OneVector;
 
-	/** The actor to track when mode is set to target */
-	UPROPERTY(EditAnywhere, BlueprintReadWrite, Setter, Getter, Category="Mode", meta=(DisplayName="TargetActor", EditCondition="Mode == ECEClonerEffectorMode::Target", EditConditionHides))
+	UE_DEPRECATED(5.5, "Moved to component and extensions")
+	UPROPERTY()
 	TWeakObjectPtr<AActor> TargetActorWeak = nullptr;
 
+	UE_DEPRECATED(5.5, "Moved to component and extensions")
 	UPROPERTY()
 	TWeakObjectPtr<AActor> InternalTargetActorWeak = nullptr;
 
-	/** Amplitude of the noise field for location */
-	UPROPERTY(EditAnywhere, BlueprintReadWrite, Setter, Getter, Category="Mode", meta=(EditCondition="Mode == ECEClonerEffectorMode::NoiseField", EditConditionHides))
+	UE_DEPRECATED(5.5, "Moved to component and extensions")
+	UPROPERTY()
 	FVector LocationStrength = FVector::ZeroVector;
 
-	/** Amplitude of the noise field for rotation */
-	UPROPERTY(EditAnywhere, BlueprintReadWrite, Setter, Getter, Category="Mode", meta=(EditCondition="Mode == ECEClonerEffectorMode::NoiseField", EditConditionHides))
+	UE_DEPRECATED(5.5, "Moved to component and extensions")
+	UPROPERTY()
 	FRotator RotationStrength = FRotator::ZeroRotator;
 
-	/** Amplitude of the noise field for scale */
-	UPROPERTY(EditAnywhere, BlueprintReadWrite, Setter, Getter, Category="Mode", meta=(ClampMin="0", AllowPreserveRatio, Delta="0.01", EditCondition="Mode == ECEClonerEffectorMode::NoiseField", EditConditionHides))
+	UE_DEPRECATED(5.5, "Moved to component and extensions")
+	UPROPERTY()
 	FVector ScaleStrength = FVector::OneVector;
 
-	/** Panning to offset the noise field sampling */
-	UPROPERTY(EditAnywhere, BlueprintReadWrite, Setter, Getter, Category="Mode", meta=(EditCondition="Mode == ECEClonerEffectorMode::NoiseField", EditConditionHides))
+	UE_DEPRECATED(5.5, "Moved to component and extensions")
+	UPROPERTY()
 	FVector Pan = FVector::ZeroVector;
 
-	/** Intensity of the noise field */
-	UPROPERTY(EditAnywhere, BlueprintReadWrite, Setter, Getter, Category="Mode", meta=(ClampMin="0", EditCondition="Mode == ECEClonerEffectorMode::NoiseField", EditConditionHides))
+	UE_DEPRECATED(5.5, "Moved to component and extensions")
+	UPROPERTY()
 	float Frequency = 0.5f;
 
-	/** Strength of the push effect */
-	UPROPERTY(EditAnywhere, BlueprintReadWrite, Setter, Getter, Category="Mode", meta=(AllowPreserveRatio, EditCondition="Mode == ECEClonerEffectorMode::Push", EditConditionHides))
+	UE_DEPRECATED(5.5, "Moved to component and extensions")
+	UPROPERTY()
 	FVector PushStrength = FVector(100.f);
 
-	/** Relative direction computed for the push effect on each clone */
-	UPROPERTY(EditAnywhere, BlueprintReadWrite, Setter, Getter, Category="Mode", meta=(EditCondition="Mode == ECEClonerEffectorMode::Push", EditConditionHides))
+	UE_DEPRECATED(5.5, "Moved to component and extensions")
+	UPROPERTY()
 	ECEClonerEffectorPushDirection PushDirection = ECEClonerEffectorPushDirection::Forward;
 
-	/** Enable orientation force to allow each clone instance to rotate around its pivot */
-	UPROPERTY(EditAnywhere, BlueprintReadWrite, Setter="SetOrientationForceEnabled", Getter="GetOrientationForceEnabled", Category="Force")
+	UE_DEPRECATED(5.5, "Moved to component and extensions")
+	UPROPERTY()
 	bool bOrientationForceEnabled = false;
 
-	UPROPERTY(EditAnywhere, BlueprintReadWrite, Setter, Getter, Category="Force", meta=(Delta="0.0001", ClampMin="0", EditCondition="bOrientationForceEnabled", EditConditionHides))
+	UE_DEPRECATED(5.5, "Moved to component and extensions")
+	UPROPERTY()
 	float OrientationForceRate = 1.f;
 
-	UPROPERTY(EditAnywhere, BlueprintReadWrite, Setter, Getter, Category="Force", meta=(EditCondition="bOrientationForceEnabled", EditConditionHides))
+	UE_DEPRECATED(5.5, "Moved to component and extensions")
+	UPROPERTY()
 	FVector OrientationForceMin = FVector(-0.1f);
 
-	UPROPERTY(EditAnywhere, BlueprintReadWrite, Setter, Getter, Category="Force", meta=(EditCondition="bOrientationForceEnabled", EditConditionHides))
+	UE_DEPRECATED(5.5, "Moved to component and extensions")
+	UPROPERTY()
 	FVector OrientationForceMax = FVector(0.1f);
 
-	/** Enable vortex force to allow each clone instance to rotate around a specific axis on the cloner pivot */
-	UPROPERTY(EditAnywhere, BlueprintReadWrite, Setter="SetVortexForceEnabled", Getter="GetVortexForceEnabled", Category="Force")
+	UE_DEPRECATED(5.5, "Moved to component and extensions")
+	UPROPERTY()
 	bool bVortexForceEnabled = false;
 
-	UPROPERTY(EditAnywhere, BlueprintReadWrite, Setter, Getter, Category="Force", meta=(EditCondition="bVortexForceEnabled", EditConditionHides))
+	UE_DEPRECATED(5.5, "Moved to component and extensions")
+	UPROPERTY()
 	float VortexForceAmount = 10000.f;
 
-	UPROPERTY(EditAnywhere, BlueprintReadWrite, Setter, Getter, Category="Force", meta=(EditCondition="bVortexForceEnabled", EditConditionHides))
+	UE_DEPRECATED(5.5, "Moved to component and extensions")
+	UPROPERTY()
 	FVector VortexForceAxis = FVector::ZAxisVector;
 
-	/** Enable curl noise force to allow each clone instance to add random location variation */
-	UPROPERTY(EditAnywhere, BlueprintReadWrite, Setter="SetCurlNoiseForceEnabled", Getter="GetCurlNoiseForceEnabled", Category="Force")
+	UE_DEPRECATED(5.5, "Moved to component and extensions")
+	UPROPERTY()
 	bool bCurlNoiseForceEnabled = false;
 
-	UPROPERTY(EditAnywhere, BlueprintReadWrite, Setter, Getter, Category="Force", meta=(EditCondition="bCurlNoiseForceEnabled", EditConditionHides))
+	UE_DEPRECATED(5.5, "Moved to component and extensions")
+	UPROPERTY()
 	float CurlNoiseForceStrength = 1000.f;
 
-	UPROPERTY(EditAnywhere, BlueprintReadWrite, Setter, Getter, Category="Force", meta=(EditCondition="bCurlNoiseForceEnabled", EditConditionHides))
+	UE_DEPRECATED(5.5, "Moved to component and extensions")
+	UPROPERTY()
 	float CurlNoiseForceFrequency = 10.f;
 
-	/** Enable attraction force to allow each clone instances to gravitate toward a location */
-	UPROPERTY(EditAnywhere, BlueprintReadWrite, Setter="SetAttractionForceEnabled", Getter="GetAttractionForceEnabled", Category="Force")
+	UE_DEPRECATED(5.5, "Moved to component and extensions")
+	UPROPERTY()
 	bool bAttractionForceEnabled = false;
 
-	UPROPERTY(EditAnywhere, BlueprintReadWrite, Setter, Getter, Category="Force", meta=(EditCondition="bAttractionForceEnabled", EditConditionHides))
+	UE_DEPRECATED(5.5, "Moved to component and extensions")
+	UPROPERTY()
 	float AttractionForceStrength = 1000.f;
 
-	UPROPERTY(EditAnywhere, BlueprintReadWrite, Setter, Getter, Category="Force", meta=(EditCondition="bAttractionForceEnabled", EditConditionHides))
+	UE_DEPRECATED(5.5, "Moved to component and extensions")
+	UPROPERTY()
 	float AttractionForceFalloff = 0.1f;
 
-	/** Enable gravity force to pull particles based on an acceleration vector */
-	UPROPERTY(EditAnywhere, BlueprintReadWrite, Setter="SetGravityForceEnabled", Getter="GetGravityForceEnabled", Category="Force")
+	UE_DEPRECATED(5.5, "Moved to component and extensions")
+	UPROPERTY()
 	bool bGravityForceEnabled = false;
 
-	UPROPERTY(EditAnywhere, BlueprintReadWrite, Setter, Getter, Category="Force", meta=(EditCondition="bGravityForceEnabled", EditConditionHides))
+	UE_DEPRECATED(5.5, "Moved to component and extensions")
+	UPROPERTY()
 	FVector GravityForceAcceleration = FVector(0, 0, -980.f);
 
 #if WITH_EDITORONLY_DATA
-	/** Visibility of the components visualizer */
-	UPROPERTY(EditAnywhere, BlueprintReadWrite, AdvancedDisplay, Category="Effector")
+	UE_DEPRECATED(5.5, "Moved to component and extensions")
+	UPROPERTY()
 	bool bVisualizerComponentVisible = true;
 
-	/** Toggle the sprite to visualize and click on this effector */
-	UPROPERTY(EditAnywhere, BlueprintReadWrite, AdvancedDisplay, Category="Effector")
+	UE_DEPRECATED(5.5, "Moved to component and extensions")
+	UPROPERTY()
 	bool bVisualizerSpriteVisible = true;
 #endif
-
-private:
-	void OnTargetActorTransformChanged(USceneComponent* InUpdatedComponent, EUpdateTransformFlags InFlags, ETeleportType InTeleport);
-
-	UFUNCTION()
-	void OnTargetActorDestroyed(AActor* InActor);
-
-	/** Visualizer component for inner effector type */
-	UPROPERTY(Transient)
-	TObjectPtr<UDynamicMeshComponent> InnerVisualizerComponent;
-
-	/** Visualizer component for outer effector type */
-	UPROPERTY(Transient)
-	TObjectPtr<UDynamicMeshComponent> OuterVisualizerComponent;
-
-#if WITH_EDITORONLY_DATA
-	/** Dynamic material for the inner visualizer */
-	UPROPERTY(Transient)
-	TObjectPtr<UMaterialInstanceDynamic> InnerVisualizerMaterial;
-
-	/** Dynamic material for the outer visualizer */
-	UPROPERTY(Transient)
-	TObjectPtr<UMaterialInstanceDynamic> OuterVisualizerMaterial;
-
-	/** Used to hold data related to visualizer for updates */
-	UPROPERTY(Transient, NonTransactional)
-	FInstancedPropertyBag VisualizerData;
-#endif
-
-	/** Transient effector channel data */
-	UPROPERTY(VisibleInstanceOnly, Transient, DuplicateTransient, TextExportTransient, NonTransactional, AdvancedDisplay, Category="Effector")
-	FCEClonerEffectorChannelData ChannelData;
-
-	/** Cloners linked to this effector, used to refresh or relink on duplicate */
-	UPROPERTY(SkipSerialization, NonTransactional)
-	TArray<TWeakObjectPtr<ACEClonerActor>> ClonersWeak;
-
-#if WITH_EDITOR
-	/** Used for PECP */
-	static TCEPropertyChangeDispatcher<ACEEffectorActor> PropertyChangeDispatcher;
-#endif
 };
diff --git a/Engine/Plugins/Experimental/ClonerEffector/Source/ClonerEffector/Public/Effector/CEEffectorComponent.h b/Engine/Plugins/Experimental/ClonerEffector/Source/ClonerEffector/Public/Effector/CEEffectorComponent.h
--- a/Engine/Plugins/Experimental/ClonerEffector/Source/ClonerEffector/Public/Effector/CEEffectorComponent.h
+++ b/Engine/Plugins/Experimental/ClonerEffector/Source/ClonerEffector/Public/Effector/CEEffectorComponent.h
@@ -2,12 +2,269 @@
 
 #pragma once
 
+#include "CEClonerEffectorShared.h"
+#include "CEPropertyChangeDispatcher.h"
 #include "Components/SceneComponent.h"
+#include "UObject/WeakObjectPtrTemplates.h"
 #include "CEEffectorComponent.generated.h"
 
-/** Class used to define a custom visualizer for this component and the effector actor */
-UCLASS(MinimalAPI, Within=CEEffectorActor)
+class UCEClonerEffectorExtension;
+class UCEEffectorExtensionBase;
+class UCEEffectorModeBase;
+class UCEEffectorTypeBase;
+class UDynamicMesh;
+class UDynamicMeshComponent;
+
+/** Class used to define an effector */
+UCLASS(
+	MinimalAPI
+	, BlueprintType
+	, HideCategories=(Tags,Activation,LOD,AssetUserData,Navigation,Rendering,Cooking,Physics)
+	, DisplayName = "Motion Design Effector Component")
 class UCEEffectorComponent : public USceneComponent
 {
 	GENERATED_BODY()
+
+	friend class ACEEffectorActor;
+	friend class FCEEditorEffectorComponentDetailCustomization;
+	friend class UCEClonerEffectorExtension;
+
+public:
+	UCEEffectorComponent();
+
+	UFUNCTION(BlueprintCallable, Category="Effector")
+	CLONEREFFECTOR_API void SetEnabled(bool bInEnable);
+
+	UFUNCTION(BlueprintPure, Category="Effector")
+	bool GetEnabled() const
+	{
+		return bEnabled;
+	}
+
+	UFUNCTION(BlueprintCallable, Category="Effector")
+	CLONEREFFECTOR_API void SetColor(const FLinearColor& InColor);
+
+	UFUNCTION(BlueprintPure, Category="Effector")
+	const FLinearColor& GetColor() const
+	{
+		return Color;
+	}
+
+	UFUNCTION(BlueprintCallable, Category="Effector")
+	CLONEREFFECTOR_API void SetMagnitude(float InMagnitude);
+
+	UFUNCTION(BlueprintPure, Category="Effector")
+	float GetMagnitude() const
+	{
+		return Magnitude;
+	}
+
+	UFUNCTION(BlueprintCallable, Category="Effector")
+	CLONEREFFECTOR_API void SetTypeName(FName InTypeName);
+
+	UFUNCTION(BlueprintPure, Category="Effector")
+	FName GetTypeName() const
+	{
+		return TypeName;
+	}
+
+	UFUNCTION(BlueprintCallable, Category="Effector")
+	CLONEREFFECTOR_API void SetModeName(FName InModeName);
+
+	UFUNCTION(BlueprintPure, Category="Effector")
+	FName GetModeName() const
+	{
+		return ModeName;
+	}
+
+#if WITH_EDITOR
+	UFUNCTION(BlueprintCallable, Category="Effector")
+	CLONEREFFECTOR_API void SetVisualizerComponentVisible(bool bInVisible);
+
+	UFUNCTION(BlueprintPure, Category="Effector")
+	bool GetVisualizerComponentVisible() const
+	{
+		return bVisualizerComponentVisible;
+	}
+
+	UFUNCTION(BlueprintCallable, Category="Effector")
+	CLONEREFFECTOR_API void SetVisualizerSpriteVisible(bool bInVisible);
+
+	UFUNCTION(BlueprintPure, Category="Effector")
+	bool GetVisualizerSpriteVisible() const
+	{
+		return bVisualizerSpriteVisible;
+	}
+
+	int32 AddVisualizerComponent(UDynamicMeshComponent* InComponent);
+
+	/** Updates a specific mesh visualizers based on its index */
+	void UpdateVisualizer(int32 InVisualizerIndex, TFunctionRef<void(UDynamicMesh*)> InMeshFunction) const;
+#endif
+
+	/** Get linked cloners */
+	TConstArrayView<TWeakObjectPtr<UCEClonerEffectorExtension>> GetClonerExtensionsWeak() const;
+
+	/** Get the effector channel data */
+	FCEClonerEffectorChannelData& GetChannelData();
+
+	/** Get the effector channel identifier */
+	int32 GetChannelIdentifier() const;
+
+	void OnClonerLinked(UCEClonerEffectorExtension* InClonerExtension);
+	void OnClonerUnlinked(UCEClonerEffectorExtension* InClonerExtension);
+
+	template<
+		typename InTypeClass
+		UE_REQUIRES(TIsDerivedFrom<InTypeClass, UCEEffectorTypeBase>::Value)>
+	InTypeClass* GetActiveType() const
+	{
+		return Cast<InTypeClass>(GetActiveType());
+	}
+
+	UCEEffectorTypeBase* GetActiveType() const
+	{
+		return ActiveType;
+	}
+
+	template<
+		typename InModeClass
+		UE_REQUIRES(TIsDerivedFrom<InModeClass, UCEEffectorModeBase>::Value)>
+	InModeClass* GetActiveType() const
+	{
+		return Cast<InModeClass>(GetActiveMode());
+	}
+
+	UCEEffectorModeBase* GetActiveMode() const
+	{
+		return ActiveMode;
+	}
+
+	void RequestClonerUpdate(bool bInImmediate);
+
+protected:
+	//~ Begin UActorComponent
+	virtual void OnComponentCreated() override;
+	virtual void OnComponentDestroyed(bool bDestroyingHierarchy) override;
+	//~ End UActorComponent
+
+	//~ Begin UObject
+#if WITH_EDITOR
+	virtual void PostEditChangeProperty(FPropertyChangedEvent& InPropertyChangedEvent) override;
+	virtual void PostEditUndo() override;
+#endif
+	virtual void PostEditImport() override;
+	virtual void PostLoad() override;
+	virtual void PostDuplicate(EDuplicateMode::Type InDuplicateMode) override;
+	//~ End UObject
+
+#if WITH_EDITOR
+	UFUNCTION(CallInEditor, Category="Effector")
+	void ForceRefreshLinkedCloners();
+#endif
+
+	/** Is this effector enabled/disabled on linked cloners */
+	UPROPERTY(EditInstanceOnly, BlueprintReadWrite, Setter="SetEnabled", Getter="GetEnabled", Category="Effector")
+	bool bEnabled = true;
+
+	/** The ratio effect of the effector on clones */
+	UPROPERTY(EditInstanceOnly, BlueprintReadWrite, Setter, Getter, Category="Effector", meta=(ClampMin="0", ClampMax="1"))
+	float Magnitude = 1.f;
+
+	/** Affected clones color passed over to material */
+	UPROPERTY(EditInstanceOnly, BlueprintReadWrite, Setter, Getter, Category="Effector")
+	FLinearColor Color = FLinearColor::Red;
+
+	/** Name of the shape type to use */
+	UPROPERTY(EditInstanceOnly, BlueprintReadWrite, Setter, Getter, Category="Shape", meta=(GetOptions="GetEffectorTypeNames"))
+	FName TypeName = NAME_None;
+
+	/** Cached active type used for faster access */
+	UPROPERTY(VisibleInstanceOnly, BlueprintReadOnly, Instanced, Transient, DuplicateTransient, NonTransactional, Category="Shape", meta=(DisplayAfter="TypeName"))
+	TObjectPtr<UCEEffectorTypeBase> ActiveType;
+
+	/** Name of the shape type to use */
+	UPROPERTY(EditInstanceOnly, BlueprintReadWrite, Setter, Getter, Category="Mode", meta=(GetOptions="GetEffectorModeNames"))
+	FName ModeName = NAME_None;
+
+	/** Cached active mode used for faster access */
+	UPROPERTY(VisibleInstanceOnly, BlueprintReadOnly, Instanced, Transient, DuplicateTransient, NonTransactional, Category="Mode", meta=(DisplayAfter="ModeName"))
+	TObjectPtr<UCEEffectorModeBase> ActiveMode;
+
+	/** Active Extensions on this effector */
+	UPROPERTY(VisibleInstanceOnly, BlueprintReadOnly, Instanced, Transient, DuplicateTransient, NonTransactional, Category="Extensions")
+	TArray<TObjectPtr<UCEEffectorExtensionBase>> ActiveExtensions;
+
+#if WITH_EDITORONLY_DATA
+	/** Visibility of the components visualizer */
+	UPROPERTY(EditInstanceOnly, BlueprintReadWrite, AdvancedDisplay, Category="Effector")
+	bool bVisualizerComponentVisible = true;
+
+	/** Toggle the sprite to visualize and click on this effector */
+	UPROPERTY(EditInstanceOnly, BlueprintReadWrite, AdvancedDisplay, Category="Effector")
+	bool bVisualizerSpriteVisible = true;
+#endif
+
+private:
+	static constexpr TCHAR VisualizerMaterialPath[] = TEXT("/Script/Engine.Material'/ClonerEffector/Materials/M_EffectorVisualizer.M_EffectorVisualizer'");
+	static constexpr TCHAR VisualizerColorName[] = TEXT("VisualizerColor");
+	static constexpr TCHAR SpriteTexturePath[] = TEXT("/Script/Engine.Texture2D'/ClonerEffector/Textures/T_EffectorIcon.T_EffectorIcon'");
+
+	void RegisterToChannel();
+	void OnTransformUpdated(USceneComponent* InComponent, EUpdateTransformFlags InFlags, ETeleportType InTeleportType);
+
+	void OnEnabledChanged();
+	void OnEffectorEnabled();
+	void OnEffectorDisabled();
+	void OnEffectorOptionsChanged();
+	void OnTypeNameChanged();
+	void OnModeNameChanged();
+	void OnExtensionsChanged();
+
+#if WITH_EDITOR
+	void OnEffectorDeveloperSettingsChanged(UObject* InSettings, FPropertyChangedEvent& InEvent);
+	void OnVisualizerOptionsChanged();
+#endif
+
+	template<
+		typename InExtensionClass>
+	InExtensionClass* FindOrAddExtension()
+	{
+		return Cast<InExtensionClass>(FindOrAddExtension(InExtensionClass::StaticClass()));
+	}
+
+	UCEEffectorExtensionBase* FindOrAddExtension(FName InExtensionName);
+	UCEEffectorExtensionBase* FindOrAddExtension(TSubclassOf<UCEEffectorExtensionBase> InExtensionClass);
+
+	/** Gets all type names available */
+	UFUNCTION()
+	TArray<FName> GetEffectorTypeNames() const;
+
+	/** Gets all mode names available */
+	UFUNCTION()
+	TArray<FName> GetEffectorModeNames() const;
+
+	/** Transient effector channel data */
+	UPROPERTY(VisibleInstanceOnly, Transient, DuplicateTransient, TextExportTransient, NonTransactional, AdvancedDisplay, Category="Effector")
+	FCEClonerEffectorChannelData ChannelData;
+
+	/** Cloners linked to this effector, used to refresh or relink on duplicate */
+	UPROPERTY(SkipSerialization, NonTransactional)
+	TArray<TWeakObjectPtr<UCEClonerEffectorExtension>> ClonerExtensionsWeak;
+
+	/** Cached extensions instances */
+	UPROPERTY()
+	TArray<TObjectPtr<UCEEffectorExtensionBase>> ExtensionInstances;
+
+#if WITH_EDITORONLY_DATA
+	UPROPERTY(Transient, DuplicateTransient)
+	TArray<TWeakObjectPtr<UDynamicMeshComponent>> VisualizerComponentsWeak;
+
+	UPROPERTY(Transient, DuplicateTransient)
+	TArray<TObjectPtr<UMaterialInstanceDynamic>> VisualizerMaterialsWeak;
+#endif
+
+#if WITH_EDITOR
+	static const TCEPropertyChangeDispatcher<UCEEffectorComponent> PropertyChangeDispatcher;
+#endif
 };
\ No newline at end of file
diff --git a/Engine/Plugins/Experimental/ClonerEffector/Source/ClonerEffector/Public/Effector/Extensions/CEEffectorForceExtension.h b/Engine/Plugins/Experimental/ClonerEffector/Source/ClonerEffector/Public/Effector/Extensions/CEEffectorForceExtension.h
new file mode 100644
--- /dev/null
+++ b/Engine/Plugins/Experimental/ClonerEffector/Source/ClonerEffector/Public/Effector/Extensions/CEEffectorForceExtension.h
@@ -0,0 +1,238 @@
+// Copyright Epic Games, Inc. All Rights Reserved.
+
+#pragma once
+
+#include "CEPropertyChangeDispatcher.h"
+#include "Effector/CEEffectorExtensionBase.h"
+#include "CEEffectorForceExtension.generated.h"
+
+class UCEEffectorComponent;
+
+/** Effector supported forces to affect clones */
+UCLASS(MinimalAPI, BlueprintType, Within=CEEffectorComponent)
+class UCEEffectorForceExtension : public UCEEffectorExtensionBase
+{
+	GENERATED_BODY()
+
+public:
+	UCEEffectorForceExtension()
+		: UCEEffectorExtensionBase(TEXT("Forces"), TEXT("Forces"))
+	{}
+
+	UFUNCTION(BlueprintCallable, Category="Effector")
+	CLONEREFFECTOR_API void SetForcesEnabled(bool bInForcesEnabled);
+
+	UFUNCTION(BlueprintPure, Category="Effector")
+	bool GetForcesEnabled() const
+	{
+		return bForcesEnabled;
+	}
+
+	UFUNCTION(BlueprintCallable, Category="Effector")
+	CLONEREFFECTOR_API void SetOrientationForceEnabled(bool bInForceEnabled);
+
+	UFUNCTION(BlueprintPure, Category="Effector")
+	bool GetOrientationForceEnabled() const
+	{
+		return bOrientationForceEnabled;
+	}
+
+	UFUNCTION(BlueprintCallable, Category="Effector")
+	CLONEREFFECTOR_API void SetOrientationForceRate(float InForceOrientationRate);
+
+	UFUNCTION(BlueprintPure, Category="Effector")
+	float GetOrientationForceRate() const
+	{
+		return OrientationForceRate;
+	}
+
+	UFUNCTION(BlueprintCallable, Category="Effector")
+	CLONEREFFECTOR_API void SetOrientationForceMin(const FVector& InForceOrientationMin);
+
+	UFUNCTION(BlueprintPure, Category="Effector")
+	const FVector& GetOrientationForceMin() const
+	{
+		return OrientationForceMin;
+	}
+
+	UFUNCTION(BlueprintCallable, Category="Effector")
+	CLONEREFFECTOR_API void SetOrientationForceMax(const FVector& InForceOrientationMax);
+
+	UFUNCTION(BlueprintPure, Category="Effector")
+	const FVector& GetOrientationForceMax() const
+	{
+		return OrientationForceMax;
+	}
+
+	UFUNCTION(BlueprintCallable, Category="Effector")
+	CLONEREFFECTOR_API void SetVortexForceEnabled(bool bInForceEnabled);
+
+	UFUNCTION(BlueprintPure, Category="Effector")
+	bool GetVortexForceEnabled() const
+	{
+		return bVortexForceEnabled;
+	}
+
+	UFUNCTION(BlueprintCallable, Category="Effector")
+	CLONEREFFECTOR_API void SetVortexForceAmount(float InForceVortexAmount);
+
+	UFUNCTION(BlueprintPure, Category="Effector")
+	float GetVortexForceAmount() const
+	{
+		return VortexForceAmount;
+	}
+
+	UFUNCTION(BlueprintCallable, Category="Effector")
+	CLONEREFFECTOR_API void SetVortexForceAxis(const FVector& InForceVortexAxis);
+
+	UFUNCTION(BlueprintPure, Category="Effector")
+	const FVector& GetVortexForceAxis() const
+	{
+		return VortexForceAxis;
+	}
+
+	UFUNCTION(BlueprintCallable, Category="Effector")
+	CLONEREFFECTOR_API void SetCurlNoiseForceEnabled(bool bInForceEnabled);
+
+	UFUNCTION(BlueprintPure, Category="Effector")
+	bool GetCurlNoiseForceEnabled() const
+	{
+		return bCurlNoiseForceEnabled;
+	}
+
+	UFUNCTION(BlueprintCallable, Category="Effector")
+	CLONEREFFECTOR_API void SetCurlNoiseForceStrength(float InForceCurlNoiseStrength);
+
+	UFUNCTION(BlueprintPure, Category="Effector")
+	float GetCurlNoiseForceStrength() const
+	{
+		return CurlNoiseForceStrength;
+	}
+
+	UFUNCTION(BlueprintCallable, Category="Effector")
+	CLONEREFFECTOR_API void SetCurlNoiseForceFrequency(float InForceCurlNoiseFrequency);
+
+	UFUNCTION(BlueprintPure, Category="Effector")
+	float GetCurlNoiseForceFrequency() const
+	{
+		return CurlNoiseForceFrequency;
+	}
+
+	UFUNCTION(BlueprintCallable, Category="Effector")
+	CLONEREFFECTOR_API void SetAttractionForceEnabled(bool bInForceEnabled);
+
+	UFUNCTION(BlueprintPure, Category="Effector")
+	bool GetAttractionForceEnabled() const
+	{
+		return bAttractionForceEnabled;
+	}
+
+	UFUNCTION(BlueprintCallable, Category="Effector")
+	CLONEREFFECTOR_API void SetAttractionForceStrength(float InForceStrength);
+
+	UFUNCTION(BlueprintPure, Category="Effector")
+	float GetAttractionForceStrength() const
+	{
+		return AttractionForceStrength;
+	}
+
+	UFUNCTION(BlueprintCallable, Category="Effector")
+	CLONEREFFECTOR_API void SetAttractionForceFalloff(float InForceFalloff);
+
+	UFUNCTION(BlueprintPure, Category="Effector")
+	float GetAttractionForceFalloff() const
+	{
+		return AttractionForceFalloff;
+	}
+
+	UFUNCTION(BlueprintCallable, Category="Effector")
+	CLONEREFFECTOR_API void SetGravityForceEnabled(bool bInForceEnabled);
+
+	UFUNCTION(BlueprintPure, Category="Effector")
+	bool GetGravityForceEnabled() const
+	{
+		return bGravityForceEnabled;
+	}
+
+	UFUNCTION(BlueprintCallable, Category="Effector")
+	CLONEREFFECTOR_API void SetGravityForceAcceleration(const FVector& InAcceleration);
+
+	UFUNCTION(BlueprintPure, Category="Effector")
+	FVector GetGravityForceAcceleration() const
+	{
+		return GravityForceAcceleration;
+	}
+
+protected:
+	//~ Begin UObject
+#if WITH_EDITOR
+	virtual void PostEditChangeProperty(FPropertyChangedEvent& InPropertyChangedEvent) override;
+#endif
+	//~ End UObject
+
+	//~ Begin UCEEffectorExtensionBase
+	virtual void OnExtensionParametersChanged(UCEEffectorComponent* InComponent) override;
+	//~ End UCEEffectorExtensionBase
+
+	void OnForceOptionsChanged();
+
+	/** Forces global state */
+	UPROPERTY(EditInstanceOnly, BlueprintReadWrite, Setter="SetForcesEnabled", Getter="GetForcesEnabled", Category="Force")
+	bool bForcesEnabled = false;
+
+	/** Enable orientation force to allow each clone instance to rotate around its pivot */
+	UPROPERTY(EditInstanceOnly, BlueprintReadWrite, Setter="SetOrientationForceEnabled", Getter="GetOrientationForceEnabled", Category="Force", meta=(EditCondition="bForcesEnabled", EditConditionHides))
+	bool bOrientationForceEnabled = false;
+
+	UPROPERTY(EditInstanceOnly, BlueprintReadWrite, Setter, Getter, Category="Force", meta=(Delta="0.0001", ClampMin="0", EditCondition="bForcesEnabled && bOrientationForceEnabled", EditConditionHides))
+	float OrientationForceRate = 1.f;
+
+	UPROPERTY(EditInstanceOnly, BlueprintReadWrite, Setter, Getter, Category="Force", meta=(EditCondition="bForcesEnabled && bOrientationForceEnabled", EditConditionHides))
+	FVector OrientationForceMin = FVector(-0.1f);
+
+	UPROPERTY(EditInstanceOnly, BlueprintReadWrite, Setter, Getter, Category="Force", meta=(EditCondition="bForcesEnabled && bOrientationForceEnabled", EditConditionHides))
+	FVector OrientationForceMax = FVector(0.1f);
+
+	/** Enable vortex force to allow each clone instance to rotate around a specific axis on the cloner pivot */
+	UPROPERTY(EditInstanceOnly, BlueprintReadWrite, Setter="SetVortexForceEnabled", Getter="GetVortexForceEnabled", Category="Force", meta=(EditCondition="bForcesEnabled", EditConditionHides))
+	bool bVortexForceEnabled = false;
+
+	UPROPERTY(EditInstanceOnly, BlueprintReadWrite, Setter, Getter, Category="Force", meta=(EditCondition="bForcesEnabled && bVortexForceEnabled", EditConditionHides))
+	float VortexForceAmount = 10000.f;
+
+	UPROPERTY(EditInstanceOnly, BlueprintReadWrite, Setter, Getter, Category="Force", meta=(EditCondition="bForcesEnabled && bVortexForceEnabled", EditConditionHides))
+	FVector VortexForceAxis = FVector::ZAxisVector;
+
+	/** Enable curl noise force to allow each clone instance to add random location variation */
+	UPROPERTY(EditInstanceOnly, BlueprintReadWrite, Setter="SetCurlNoiseForceEnabled", Getter="GetCurlNoiseForceEnabled", Category="Force", meta=(EditCondition="bForcesEnabled", EditConditionHides))
+	bool bCurlNoiseForceEnabled = false;
+
+	UPROPERTY(EditInstanceOnly, BlueprintReadWrite, Setter, Getter, Category="Force", meta=(EditCondition="bForcesEnabled && bCurlNoiseForceEnabled", EditConditionHides))
+	float CurlNoiseForceStrength = 1000.f;
+
+	UPROPERTY(EditInstanceOnly, BlueprintReadWrite, Setter, Getter, Category="Force", meta=(EditCondition="bForcesEnabled && bCurlNoiseForceEnabled", EditConditionHides))
+	float CurlNoiseForceFrequency = 10.f;
+
+	/** Enable attraction force to allow each clone instances to gravitate toward a location */
+	UPROPERTY(EditInstanceOnly, BlueprintReadWrite, Setter="SetAttractionForceEnabled", Getter="GetAttractionForceEnabled", Category="Force", meta=(EditCondition="bForcesEnabled", EditConditionHides))
+	bool bAttractionForceEnabled = false;
+
+	UPROPERTY(EditInstanceOnly, BlueprintReadWrite, Setter, Getter, Category="Force", meta=(EditCondition="bForcesEnabled && bAttractionForceEnabled", EditConditionHides))
+	float AttractionForceStrength = 1000.f;
+
+	UPROPERTY(EditInstanceOnly, BlueprintReadWrite, Setter, Getter, Category="Force", meta=(EditCondition="bForcesEnabled && bAttractionForceEnabled", EditConditionHides))
+	float AttractionForceFalloff = 0.1f;
+
+	/** Enable gravity force to pull particles based on an acceleration vector */
+	UPROPERTY(EditInstanceOnly, BlueprintReadWrite, Setter="SetGravityForceEnabled", Getter="GetGravityForceEnabled", Category="Force", meta=(EditCondition="bForcesEnabled", EditConditionHides))
+	bool bGravityForceEnabled = false;
+
+	UPROPERTY(EditInstanceOnly, BlueprintReadWrite, Setter, Getter, Category="Force", meta=(EditCondition="bForcesEnabled && bGravityForceEnabled", EditConditionHides))
+	FVector GravityForceAcceleration = FVector(0, 0, -980.f);
+
+private:
+#if WITH_EDITOR
+	/** Used for PECP */
+	static const TCEPropertyChangeDispatcher<UCEEffectorForceExtension> PropertyChangeDispatcher;
+#endif
+};
\ No newline at end of file
diff --git a/Engine/Plugins/Experimental/ClonerEffector/Source/ClonerEffector/Public/Effector/Modes/CEEffectorNoiseMode.h b/Engine/Plugins/Experimental/ClonerEffector/Source/ClonerEffector/Public/Effector/Modes/CEEffectorNoiseMode.h
new file mode 100644
--- /dev/null
+++ b/Engine/Plugins/Experimental/ClonerEffector/Source/ClonerEffector/Public/Effector/Modes/CEEffectorNoiseMode.h
@@ -0,0 +1,103 @@
+// Copyright Epic Games, Inc. All Rights Reserved.
+
+#pragma once
+
+#include "CEClonerEffectorShared.h"
+#include "CEEffectorModeBase.h"
+#include "CEPropertyChangeDispatcher.h"
+#include "CEEffectorNoiseMode.generated.h"
+
+class UCEEffectorComponent;
+
+UCLASS(MinimalAPI, BlueprintType, Within=CEEffectorComponent)
+class UCEEffectorNoiseMode : public UCEEffectorModeBase
+{
+	GENERATED_BODY()
+
+public:
+	UCEEffectorNoiseMode()
+		: UCEEffectorModeBase(TEXT("Noise"), static_cast<int32>(ECEClonerEffectorMode::NoiseField))
+	{}
+
+	UFUNCTION(BlueprintCallable, Category="Effector")
+	CLONEREFFECTOR_API void SetLocationStrength(const FVector& InStrength);
+
+	UFUNCTION(BlueprintPure, Category="Effector")
+	FVector GetLocationStrength() const
+	{
+		return LocationStrength;
+	}
+
+	UFUNCTION(BlueprintCallable, Category="Effector")
+	CLONEREFFECTOR_API void SetRotationStrength(const FRotator& InStrength);
+
+	UFUNCTION(BlueprintPure, Category="Effector")
+	FRotator GetRotationStrength() const
+	{
+		return RotationStrength;
+	}
+
+	UFUNCTION(BlueprintCallable, Category="Effector")
+	CLONEREFFECTOR_API void SetScaleStrength(const FVector& InStrength);
+
+	UFUNCTION(BlueprintPure, Category="Effector")
+	FVector GetScaleStrength() const
+	{
+		return ScaleStrength;
+	}
+
+	UFUNCTION(BlueprintCallable, Category="Effector")
+	CLONEREFFECTOR_API void SetPan(const FVector& InPan);
+
+	UFUNCTION(BlueprintPure, Category="Effector")
+	FVector GetPan() const
+	{
+		return Pan;
+	}
+
+	UFUNCTION(BlueprintCallable, Category="Effector")
+	CLONEREFFECTOR_API void SetFrequency(float InFrequency);
+
+	UFUNCTION(BlueprintPure, Category="Effector")
+	float GetFrequency() const
+	{
+		return Frequency;
+	}
+
+protected:
+	//~ Begin UCEEffectorNoiseMode
+	virtual void OnExtensionParametersChanged(UCEEffectorComponent* InComponent) override;
+	//~ End UCEEffectorNoiseMode
+
+	//~ Begin UObject
+#if WITH_EDITOR
+	virtual void PostEditChangeProperty(FPropertyChangedEvent& InPropertyChangedEvent) override;
+#endif
+	//~ End UObject
+
+	/** Amplitude of the noise field for location */
+	UPROPERTY(EditInstanceOnly, BlueprintReadWrite, Setter, Getter, Category="Mode")
+	FVector LocationStrength = FVector::ZeroVector;
+
+	/** Amplitude of the noise field for rotation */
+	UPROPERTY(EditInstanceOnly, BlueprintReadWrite, Setter, Getter, Category="Mode")
+	FRotator RotationStrength = FRotator::ZeroRotator;
+
+	/** Amplitude of the noise field for scale */
+	UPROPERTY(EditInstanceOnly, BlueprintReadWrite, Setter, Getter, Category="Mode", meta=(ClampMin="0", AllowPreserveRatio, Delta="0.01"))
+	FVector ScaleStrength = FVector::OneVector;
+
+	/** Panning to offset the noise field sampling */
+	UPROPERTY(EditInstanceOnly, BlueprintReadWrite, Setter, Getter, Category="Mode")
+	FVector Pan = FVector::ZeroVector;
+
+	/** Intensity of the noise field */
+	UPROPERTY(EditInstanceOnly, BlueprintReadWrite, Setter, Getter, Category="Mode", meta=(ClampMin="0"))
+	float Frequency = 0.5f;
+
+private:
+#if WITH_EDITOR
+	/** Used for PECP */
+	static const TCEPropertyChangeDispatcher<UCEEffectorNoiseMode> PropertyChangeDispatcher;
+#endif
+};
\ No newline at end of file
diff --git a/Engine/Plugins/Experimental/ClonerEffector/Source/ClonerEffector/Public/Subsystems/CEClonerSubsystem.h b/Engine/Plugins/Experimental/ClonerEffector/Source/ClonerEffector/Public/Subsystems/CEClonerSubsystem.h
--- a/Engine/Plugins/Experimental/ClonerEffector/Source/ClonerEffector/Public/Subsystems/CEClonerSubsystem.h
+++ b/Engine/Plugins/Experimental/ClonerEffector/Source/ClonerEffector/Public/Subsystems/CEClonerSubsystem.h
@@ -7,44 +7,27 @@
 #include "UObject/WeakObjectPtrTemplates.h"
 #include "CEClonerSubsystem.generated.h"
 
+class UCEClonerExtensionBase;
+
 UCLASS(MinimalAPI)
 class UCEClonerSubsystem : public UEngineSubsystem
 {
 	GENERATED_BODY()
 
 public:
-	static inline constexpr int32 NoFlicker = 1;
-
-#if WITH_EDITOR
-	DECLARE_MULTICAST_DELEGATE(FOnCVarChanged)
-	static FOnCVarChanged& OnCVarChanged()
+	static TMulticastDelegateRegistration<void()>& OnSubsystemInitialized()
 	{
-		return OnCVarChangedDelegate;
-	}
-#endif
-
-	DECLARE_MULTICAST_DELEGATE(FOnSubsystemInitialized)
-	static FOnSubsystemInitialized& OnSubsystemInitialized()
-	{
-		 return OnSubsystemInitializedDelegate;
+		return OnSubsystemInitializedDelegate;
 	}
 
 	/** Get this subsystem instance */
 	CLONEREFFECTOR_API static UCEClonerSubsystem* Get();
 
-	//~ Begin USubsystem
-	virtual void Initialize(FSubsystemCollectionBase& Collection) override;
-	virtual void Deinitialize() override;
-	//~ End USubsystem
+	CLONEREFFECTOR_API bool RegisterLayoutClass(UClass* InClonerLayoutClass);
 
-	CLONEREFFECTOR_API bool RegisterLayoutClass(const UClass* InClonerLayoutClass);
-	CLONEREFFECTOR_API bool UnregisterLayoutClass(const UClass* InClonerLayoutClass);
-	CLONEREFFECTOR_API bool IsLayoutClassRegistered(const UClass* InClonerLayoutClass);
+	CLONEREFFECTOR_API bool UnregisterLayoutClass(UClass* InClonerLayoutClass);
 
-	DECLARE_DELEGATE_RetVal_OneParam(TArray<AActor*> /** Children */, FOnGetOrderedActors, const AActor* /** InParent */)
-	CLONEREFFECTOR_API void RegisterCustomActorResolver(FOnGetOrderedActors InCustomResolver);
-	CLONEREFFECTOR_API void UnregisterCustomActorResolver();
-	FOnGetOrderedActors& GetCustomActorResolver();
+	CLONEREFFECTOR_API bool IsLayoutClassRegistered(UClass* InClonerLayoutClass);
 
 	/** Get available cloner layout names to use in dropdown */
 	TArray<FName> GetLayoutNames() const;
@@ -52,38 +35,49 @@ public:
 	/** Based on a layout class, find layout name */
 	FName FindLayoutName(TSubclassOf<UCEClonerLayoutBase> InLayoutClass) const;
 
-	/** Creates a new layout instance for a cloner actor */
-	UCEClonerLayoutBase* CreateNewLayout(FName InLayoutName, ACEClonerActor* InClonerActor);
+	/** Creates a new layout instance for a cloner */
+	UCEClonerLayoutBase* CreateNewLayout(FName InLayoutName, UCEClonerComponent* InCloner);
 
-#if WITH_EDITOR
-	void EnableNoFlicker();
-	void DisableNoFlicker();
-	bool IsNoFlickerEnabled() const;
+	DECLARE_DELEGATE_RetVal_OneParam(TArray<AActor*> /** Children */, FOnGetOrderedActors, const AActor* /** InParent */)
+	CLONEREFFECTOR_API void RegisterCustomActorResolver(FOnGetOrderedActors InCustomResolver);
 
-	void OnTSRShadingRejectionFlickeringPeriodChanged(IConsoleVariable* InCVar) const;
-#endif
+	CLONEREFFECTOR_API void UnregisterCustomActorResolver();
 
-protected:
-#if WITH_EDITOR
-	static FOnCVarChanged OnCVarChangedDelegate;
-#endif
+	FOnGetOrderedActors& GetCustomActorResolver();
+
+	CLONEREFFECTOR_API bool RegisterExtensionClass(UClass* InClass);
+
+	CLONEREFFECTOR_API bool UnregisterExtensionClass(UClass* InClass);
+
+	CLONEREFFECTOR_API bool IsExtensionClassRegistered(UClass* InClass) const;
+
+	/** Get available cloner extension names to use */
+	TArray<FName> GetExtensionNames() const;
 
+	/** Based on a extension class, find extension name */
+	FName FindExtensionName(TSubclassOf<UCEClonerExtensionBase> InClass) const;
+
+	/** Creates a new extension instance for a cloner */
+	UCEClonerExtensionBase* CreateNewExtension(FName InExtensionName, UCEClonerComponent* InCloner);
+
+protected:
+	DECLARE_MULTICAST_DELEGATE(FOnSubsystemInitialized)
 	CLONEREFFECTOR_API static FOnSubsystemInitialized OnSubsystemInitializedDelegate;
 
+	//~ Begin USubsystem
+	virtual void Initialize(FSubsystemCollectionBase& Collection) override;
+	//~ End USubsystem
+
 	void ScanForRegistrableClasses();
 
 	/** Linking name to the layout class */
 	UPROPERTY()
 	TMap<FName, TSubclassOf<UCEClonerLayoutBase>> LayoutClasses;
 
+	/** Linking name to the extension class */
+	UPROPERTY()
+	TMap<FName, TSubclassOf<UCEClonerExtensionBase>> ExtensionClasses;
+
 	/** Used to gather ordered actors based on parent */
 	FOnGetOrderedActors ActorResolver;
-
-#if WITH_EDITOR
-	/** Allows to reduce ghosting artifacts when moving clone instances */
-	IConsoleVariable* CVarTSRShadingRejectionFlickeringPeriod = nullptr;
-
-	/** Previous value to restore it when disabled */
-	TOptional<int32> PreviousCVarValue;
-#endif
 };
\ No newline at end of file
diff --git a/Engine/Plugins/Experimental/ClonerEffector/Source/ClonerEffectorEditor/Private/CEEditorModule.cpp b/Engine/Plugins/Experimental/ClonerEffector/Source/ClonerEffectorEditor/Private/CEEditorModule.cpp
--- a/Engine/Plugins/Experimental/ClonerEffector/Source/ClonerEffectorEditor/Private/CEEditorModule.cpp
+++ b/Engine/Plugins/Experimental/ClonerEffector/Source/ClonerEffectorEditor/Private/CEEditorModule.cpp
@@ -2,10 +2,12 @@
 
 #include "CEEditorModule.h"
 
-#include "Cloner/CEClonerActor.h"
-#include "Cloner/CEEditorClonerDetailCustomization.h"
-#include "Effector/CEEditorEffectorDetailCustomization.h"
-#include "Effector/CEEffectorActor.h"
+#include "Cloner/CEEditorClonerComponentDetailCustomization.h"
+#include "Cloner/CEClonerComponent.h"
+#include "Effector/CEEditorEffectorComponentDetailCustomization.h"
+#include "Effector/CEEditorEffectorTypeDetailCustomization.h"
+#include "Effector/CEEffectorComponent.h"
+#include "Effector/Types/CEEffectorBoundType.h"
 #include "Modules/ModuleManager.h"
 #include "PropertyEditorModule.h"
 #include "Styles/CEEditorStyle.h"
@@ -18,8 +20,9 @@ void FCEEditorModule::StartupModule()
 	FCEEditorStyle::Get();
 
 	// Cloner/effector customization
-	PropertyModule.RegisterCustomClassLayout(ACEEffectorActor::StaticClass()->GetFName(), FOnGetDetailCustomizationInstance::CreateStatic(&FCEEditorEffectorDetailCustomization::MakeInstance));
-	PropertyModule.RegisterCustomClassLayout(ACEClonerActor::StaticClass()->GetFName(), FOnGetDetailCustomizationInstance::CreateStatic(&FCEEditorClonerDetailCustomization::MakeInstance));
+	PropertyModule.RegisterCustomClassLayout(UCEClonerComponent::StaticClass()->GetFName(), FOnGetDetailCustomizationInstance::CreateStatic(&FCEEditorClonerComponentDetailCustomization::MakeInstance));
+	PropertyModule.RegisterCustomClassLayout(UCEEffectorComponent::StaticClass()->GetFName(), FOnGetDetailCustomizationInstance::CreateStatic(&FCEEditorEffectorComponentDetailCustomization::MakeInstance));
+	PropertyModule.RegisterCustomClassLayout(UCEEffectorBoundType::StaticClass()->GetFName(), FOnGetDetailCustomizationInstance::CreateStatic(&FCEEditorEffectorTypeDetailCustomization::MakeInstance));
 }
 
 void FCEEditorModule::ShutdownModule()
@@ -29,8 +32,9 @@ void FCEEditorModule::ShutdownModule()
 		FPropertyEditorModule& PropertyModule = FModuleManager::GetModuleChecked<FPropertyEditorModule>(PropertyEditorName);
 
 		// Cloner/effector
-		PropertyModule.UnregisterCustomClassLayout(ACEEffectorActor::StaticClass()->GetFName());
-		PropertyModule.UnregisterCustomClassLayout(ACEClonerActor::StaticClass()->GetFName());
+		PropertyModule.UnregisterCustomClassLayout(UCEClonerComponent::StaticClass()->GetFName());
+		PropertyModule.UnregisterCustomClassLayout(UCEEffectorComponent::StaticClass()->GetFName());
+		PropertyModule.UnregisterCustomClassLayout(UCEEffectorBoundType::StaticClass()->GetFName());
 	}
 }
 
diff --git a/Engine/Plugins/Experimental/ClonerEffector/Source/ClonerEffectorEditor/Private/CEEditorModule.h b/Engine/Plugins/Experimental/ClonerEffector/Source/ClonerEffectorEditor/Private/CEEditorModule.h
--- a/Engine/Plugins/Experimental/ClonerEffector/Source/ClonerEffectorEditor/Private/CEEditorModule.h
+++ b/Engine/Plugins/Experimental/ClonerEffector/Source/ClonerEffectorEditor/Private/CEEditorModule.h
@@ -7,7 +7,7 @@
 class FCEEditorModule : public IModuleInterface
 {
 public:
-	static inline constexpr TCHAR PropertyEditorName[] = TEXT("PropertyEditor");
+	static constexpr TCHAR PropertyEditorName[] = TEXT("PropertyEditor");
 
 	//~ Begin IModuleInterface
     virtual void StartupModule() override;
diff --git a/Engine/Plugins/Experimental/ClonerEffector/Source/ClonerEffectorEditor/Private/Effector/CEEditorEffectorComponentDetailCustomization.cpp b/Engine/Plugins/Experimental/ClonerEffector/Source/ClonerEffectorEditor/Private/Effector/CEEditorEffectorComponentDetailCustomization.cpp
new file mode 100644
--- /dev/null
+++ b/Engine/Plugins/Experimental/ClonerEffector/Source/ClonerEffectorEditor/Private/Effector/CEEditorEffectorComponentDetailCustomization.cpp
@@ -0,0 +1,94 @@
+// Copyright Epic Games, Inc. All Rights Reserved.
+
+#include "Effector/CEEditorEffectorComponentDetailCustomization.h"
+
+#include "DetailLayoutBuilder.h"
+#include "DetailWidgetRow.h"
+#include "Effector/CEEffectorComponent.h"
+#include "Effector/CEEffectorExtensionBase.h"
+#include "Modules/ModuleManager.h"
+#include "PropertyEditorModule.h"
+#include "Widgets/Input/SButton.h"
+
+#define LOCTEXT_NAMESPACE "CEEditorEffectorComponentDetailCustomization"
+
+void FCEEditorEffectorComponentDetailCustomization::CustomizeDetails(IDetailLayoutBuilder& InDetailBuilder)
+{
+	// Remove extension array property
+	TSharedRef<IPropertyHandle> ActiveExtensionsProperty = InDetailBuilder.GetProperty(GET_MEMBER_NAME_CHECKED(UCEEffectorComponent, ActiveExtensions), UCEEffectorComponent::StaticClass());
+
+	if (!ActiveExtensionsProperty->IsValidHandle())
+	{
+		return;
+	}
+
+	InDetailBuilder.HideProperty(ActiveExtensionsProperty);
+
+	TFunction<void(const TSharedPtr<IPropertyHandle>&, IDetailCategoryBuilder&)> AddObjectProperty = [&AddObjectProperty](const TSharedPtr<IPropertyHandle>& InPropertyHandle, IDetailCategoryBuilder& InCategoryBuilder)
+	{
+		uint32 NumChildren = 0;
+		InPropertyHandle->GetNumChildren(NumChildren);
+
+		for (uint32 ChildrenIndex = 0; ChildrenIndex < NumChildren; ChildrenIndex++)
+		{
+			TSharedPtr<IPropertyHandle> ChildHandle = InPropertyHandle->GetChildHandle(ChildrenIndex);
+
+			if (ChildHandle->IsCategoryHandle() && NumChildren == 1)
+			{
+				AddObjectProperty(ChildHandle, InCategoryBuilder);
+			}
+			else
+			{
+				InCategoryBuilder.AddProperty(ChildHandle);
+			}
+		}
+	};
+
+	TArray<FName> CategoriesPriority
+	{
+		TEXT("Effector"),
+		TEXT("Shape"),
+		TEXT("Mode")
+	};
+
+	uint32 NumElements = 0;
+	TSharedPtr<IPropertyHandleArray> ActiveExtensionArrayProperty = ActiveExtensionsProperty->AsArray();
+	ActiveExtensionArrayProperty->GetNumElements(NumElements);
+
+	for (uint32 ElementIndex = 0; ElementIndex < NumElements; ElementIndex++)
+	{
+		TSharedPtr<IPropertyHandle> ActiveExtensionProperty = ActiveExtensionArrayProperty->GetElement(ElementIndex);
+
+		UObject* ObjectValue;
+		ActiveExtensionProperty->GetValue(ObjectValue);
+
+		UCEEffectorExtensionBase* ActiveExtension = Cast<UCEEffectorExtensionBase>(ObjectValue);
+
+		if (!ActiveExtension)
+		{
+			continue;
+		}
+
+		const FName ExtensionCategory = ActiveExtension->GetExtensionCategory();
+		CategoriesPriority.AddUnique(ExtensionCategory);
+		IDetailCategoryBuilder& CategoryBuilder = InDetailBuilder.EditCategory(ExtensionCategory);
+		AddObjectProperty(ActiveExtensionProperty->GetChildHandle(0), CategoryBuilder);
+	}
+
+	static const FName PropertyEditor("PropertyEditor");
+	FPropertyEditorModule& PropertyModule = FModuleManager::GetModuleChecked<FPropertyEditorModule>(PropertyEditor);
+	static const FName ClassName = UCEEffectorComponent::StaticClass()->GetFName();
+
+	int32 SortOrder = InDetailBuilder.EditCategory(TEXT("Transform")).GetSortOrder();
+
+	for (const FName& Category : CategoriesPriority)
+	{
+		IDetailCategoryBuilder& CategoryBuilder = InDetailBuilder.EditCategory(Category);
+		CategoryBuilder.SetSortOrder(SortOrder++);
+
+		const TSharedRef<FPropertySection> Section = PropertyModule.FindOrCreateSection(ClassName, Category, FText::FromName(Category));
+		Section->AddCategory(Category);
+	}
+}
+
+#undef LOCTEXT_NAMESPACE
