MotionDesign : ClonerEffector
- Fix horde issue and suppress PVS warning

diff --git a/Engine/Plugins/Experimental/ClonerEffector/Source/ClonerEffector/Private/Cloner/CEClonerComponent.cpp b/Engine/Plugins/Experimental/ClonerEffector/Source/ClonerEffector/Private/Cloner/CEClonerComponent.cpp
--- a/Engine/Plugins/Experimental/ClonerEffector/Source/ClonerEffector/Private/Cloner/CEClonerComponent.cpp
+++ b/Engine/Plugins/Experimental/ClonerEffector/Source/ClonerEffector/Private/Cloner/CEClonerComponent.cpp
@@ -812,7 +812,8 @@ void UCEClonerComponent::UpdateActorBakedDynamicMesh(AActor* InActor)
 	AttachmentItem->BakedMaterials = MoveTemp(MeshMaterials);
 
 	// Was the mesh invalidated during the update process ?
-	AttachmentItem->MeshStatus = AttachmentItem->MeshStatus == ECEClonerAttachmentStatus::Outdated ? ECEClonerAttachmentStatus::Outdated : ECEClonerAttachmentStatus::Updated;
+	AttachmentItem->MeshStatus = AttachmentItem->MeshStatus == ECEClonerAttachmentStatus::Outdated ? ECEClonerAttachmentStatus::Outdated : ECEClonerAttachmentStatus::Updated; //-V547
+
 	InvalidateBakedStaticMesh(AttachmentItem->ItemActor.Get());
 }
 
diff --git a/Engine/Plugins/Experimental/ClonerEffector/Source/ClonerEffector/Private/Effector/CEEffectorComponent.cpp b/Engine/Plugins/Experimental/ClonerEffector/Source/ClonerEffector/Private/Effector/CEEffectorComponent.cpp
--- a/Engine/Plugins/Experimental/ClonerEffector/Source/ClonerEffector/Private/Effector/CEEffectorComponent.cpp
+++ b/Engine/Plugins/Experimental/ClonerEffector/Source/ClonerEffector/Private/Effector/CEEffectorComponent.cpp
@@ -6,6 +6,7 @@
 #include "Cloner/Extensions/CEClonerEffectorExtension.h"
 #include "Components/BillboardComponent.h"
 #include "Components/DynamicMeshComponent.h"
+#include "Containers/Ticker.h"
 #include "Effector/CEEffectorExtensionBase.h"
 #include "Effector/Modes/CEEffectorModeBase.h"
 #include "Effector/Types/CEEffectorTypeBase.h"
