MotionDesign : ClonerEffector
- Improving mesh builder to reduce conversion needed by reusing same mesh when possible
- Using niagara simulation cache to extract data and support conversion to static, dynamic, instanced static meshes, you can select the mesh asset location
- Using post load custom ticker for cloner update instead of forcing the system to be solo

diff --git a/Engine/Plugins/Experimental/ClonerEffector/Source/ClonerEffector/Private/CEClonerEffectorShared.cpp b/Engine/Plugins/Experimental/ClonerEffector/Source/ClonerEffector/Private/CEClonerEffectorShared.cpp
--- a/Engine/Plugins/Experimental/ClonerEffector/Source/ClonerEffector/Private/CEClonerEffectorShared.cpp
+++ b/Engine/Plugins/Experimental/ClonerEffector/Source/ClonerEffector/Private/CEClonerEffectorShared.cpp
@@ -2,58 +2,27 @@
 
 #include "CEClonerEffectorShared.h"
 
-#include "Components/BillboardComponent.h"
+#include "Cloner/CEClonerComponent.h"
+#include "Components/InstancedStaticMeshComponent.h"
+#include "DynamicMeshActor.h"
+#include "Engine/StaticMeshActor.h"
 #include "Engine/Texture2D.h"
+#include "Engine/World.h"
 #include "NiagaraDataChannelAccessor.h"
 #include "NiagaraSystem.h"
 #include "PrimitiveSceneProxy.h"
 #include "SceneManagement.h"
-#include "UObject/Package.h"
 #include "UObject/UObjectGlobals.h"
 
-void UE::ClonerEffector::SetBillboardComponentSprite(const AActor* InActor, const FString& InTexturePath)
-{
-	if (!InActor)
-	{
-		return;
-	}
-
-	UBillboardComponent* BillboardComponent = InActor->FindComponentByClass<UBillboardComponent>();
-
-	if (!BillboardComponent)
-	{
-		return;
-	}
-
-	UTexture2D* SpriteTexture = FindObject<UTexture2D>(nullptr, *InTexturePath);
-
-	if (!SpriteTexture)
-	{
-		SpriteTexture = LoadObject<UTexture2D>(nullptr, *InTexturePath);
-	}
-
-	if (SpriteTexture && SpriteTexture != BillboardComponent->Sprite)
-	{
-		BillboardComponent->SetSprite(SpriteTexture);
-	}
-}
-
-void UE::ClonerEffector::SetBillboardComponentVisibility(const AActor* InActor, bool bInVisibility)
-{
-	if (!InActor)
-	{
-		return;
-	}
-
-	UBillboardComponent* BillboardComponent = InActor->FindComponentByClass<UBillboardComponent>();
+#if WITH_EDITOR
+#include "AssetToolsModule.h"
+#include "AssetRegistry/AssetRegistryModule.h"
+#include "Dialogs/DlgPickPath.h"
+#include "Editor/EditorEngine.h"
+#include "UObject/Package.h"
+#endif
 
-	if (!BillboardComponent)
-	{
-		return;
-	}
-
-	BillboardComponent->SetVisibility(bInVisibility, false);
-}
+#define LOCTEXT_NAMESPACE "CEClonerEffectorShared"
 
 void FCEClonerEffectorChannelData::Write(UNiagaraDataChannelWriter* InWriter) const
 {
@@ -212,3 +181,667 @@ UNiagaraDataInterfaceArrayInt32* FCEClonerEffectorDataInterfaces::GetIndexArray(
 
 	return nullptr;
 }
+
+AStaticMeshActor* UE::ClonerEffector::Conversion::ConvertClonerToStaticMesh(UCEClonerComponent* InCloner)
+{
+	AStaticMeshActor* NewActor = nullptr;
+
+	if (!IsValid(InCloner))
+	{
+		return NewActor;
+	}
+
+	UWorld* World = InCloner->GetWorld();
+	AActor* Owner = InCloner->GetOwner();
+
+	if (!IsValid(World)
+		|| !IsValid(Owner)
+#if WITH_EDITOR
+		|| Owner->bIsEditorPreviewActor
+#endif
+		)
+	{
+		return NewActor;
+	}
+
+	FCEClonerMeshBuilder ClonerMeshBuilder;
+
+	if (!ClonerMeshBuilder.AppendComponent(InCloner, InCloner->GetComponentTransform()))
+	{
+		return NewActor;
+	}
+
+	if (ClonerMeshBuilder.GetMeshInstanceCount() == 0)
+	{
+		return NewActor;
+	}
+
+#if WITH_EDITOR
+	UPackage* ClonerPackage = InCloner->GetPackage();
+
+	if (!ClonerPackage)
+	{
+		return NewActor;
+	}
+
+	FString PackagePath;
+	if (!PickAssetPath(ClonerPackage->GetLoadedPath().GetPackageName(), PackagePath))
+	{
+		return NewActor;
+	}
+#endif
+
+	// Spawn attached actor with same flags as this actor
+	FActorSpawnParameters SpawnParameters;
+	SpawnParameters.Owner = Owner;
+	SpawnParameters.ObjectFlags = InCloner->GetFlags();
+#if WITH_EDITOR
+	SpawnParameters.bTemporaryEditorActor = false;
+#endif
+
+	const FVector ClonerLocation = InCloner->GetComponentLocation();
+	const FRotator ClonerRotation = InCloner->GetComponentRotation();
+
+	NewActor = World->SpawnActor<AStaticMeshActor>(ClonerLocation, ClonerRotation, SpawnParameters);
+
+	if (!NewActor)
+	{
+		return NewActor;
+	}
+
+	NewActor->SetMobility(EComponentMobility::Movable);
+
+	UStaticMeshComponent* StaticMeshComponent = NewActor->GetStaticMeshComponent();
+
+	UStaticMesh* StaticMesh = nullptr;
+
+#if WITH_EDITOR
+	const FString AssetPath = PackagePath + TEXT("SM_") + Owner->GetActorNameOrLabel() + TEXT("_Merged_") + FString::FromInt(NewActor->GetUniqueID());
+	StaticMesh = CreateAssetPackage<UStaticMesh>(AssetPath);
+#else
+	StaticMesh = NewObject<UStaticMesh>(StaticMeshComponent);
+#endif
+
+	if (!StaticMesh)
+	{
+		return NewActor;
+	}
+
+	TArray<TWeakObjectPtr<UMaterialInterface>> MeshMaterials;
+	FCEClonerMeshBuilder::FCEClonerMeshBuilderParams Params;
+	Params.bMergeMaterials = true;
+	ClonerMeshBuilder.BuildStaticMesh(StaticMesh, MeshMaterials, Params);
+
+#if WITH_EDITOR
+	// Replace material references that are no assets to avoid save issue in new packages
+	for (int32 Index = 0; Index < StaticMesh->GetNumSections(/** LOD */0); Index++)
+	{
+		UMaterialInterface* Material = StaticMesh->GetMaterial(Index);
+
+		if (Material && !Material->IsAsset())
+		{
+			Material = nullptr;
+		}
+
+		StaticMesh->SetMaterial(Index, Material);
+	}
+#endif
+
+	StaticMeshComponent->SetStaticMesh(StaticMesh);
+
+	for (int32 Index = 0; Index < MeshMaterials.Num(); Index++)
+	{
+		StaticMeshComponent->SetMaterial(Index, MeshMaterials[Index].Get());
+	}
+
+#if WITH_EDITOR
+	FActorLabelUtilities::SetActorLabelUnique(NewActor, Owner->GetActorNameOrLabel() + TEXT("_SM_Merged"));
+#endif
+
+	return NewActor;
+}
+
+ADynamicMeshActor* UE::ClonerEffector::Conversion::ConvertClonerToDynamicMesh(UCEClonerComponent* InCloner)
+{
+	ADynamicMeshActor* NewActor = nullptr;
+
+	if (!IsValid(InCloner))
+	{
+		return NewActor;
+	}
+
+	UWorld* World = InCloner->GetWorld();
+	AActor* Owner = InCloner->GetOwner();
+
+	if (!IsValid(World)
+		|| !IsValid(Owner)
+#if WITH_EDITOR
+		|| Owner->bIsEditorPreviewActor
+#endif
+		)
+	{
+		return NewActor;
+	}
+
+	FCEClonerMeshBuilder ClonerMeshBuilder;
+
+	if (!ClonerMeshBuilder.AppendComponent(InCloner, InCloner->GetComponentTransform()))
+	{
+		return NewActor;
+	}
+
+	if (ClonerMeshBuilder.GetMeshInstanceCount() == 0)
+	{
+		return NewActor;
+	}
+
+	// Spawn attached actor with same flags as this actor
+	FActorSpawnParameters SpawnParameters;
+	SpawnParameters.Owner = Owner;
+	SpawnParameters.ObjectFlags = InCloner->GetFlags();
+#if WITH_EDITOR
+	SpawnParameters.bTemporaryEditorActor = false;
+#endif
+
+	const FVector ClonerLocation = InCloner->GetComponentLocation();
+	const FRotator ClonerRotation = InCloner->GetComponentRotation();
+
+	NewActor = World->SpawnActor<ADynamicMeshActor>(ClonerLocation, ClonerRotation, SpawnParameters);
+
+	if (!NewActor)
+	{
+		return NewActor;
+	}
+
+	UDynamicMeshComponent* DynamicMeshComponent = NewActor->GetDynamicMeshComponent();
+
+	TArray<TWeakObjectPtr<UMaterialInterface>> MeshMaterials;
+	FCEClonerMeshBuilder::FCEClonerMeshBuilderParams Params;
+	Params.bMergeMaterials = true;
+	ClonerMeshBuilder.BuildDynamicMesh(DynamicMeshComponent->GetDynamicMesh(), MeshMaterials, Params);
+
+	for (int32 Index = 0; Index < MeshMaterials.Num(); Index++)
+	{
+		DynamicMeshComponent->SetMaterial(Index, MeshMaterials[Index].Get());
+	}
+
+#if WITH_EDITOR
+	FActorLabelUtilities::SetActorLabelUnique(NewActor, Owner->GetActorNameOrLabel() + TEXT("_DM_Merged"));
+#endif
+
+	return NewActor;
+}
+
+TArray<AStaticMeshActor*> UE::ClonerEffector::Conversion::ConvertClonerToStaticMeshes(UCEClonerComponent* InCloner)
+{
+	TArray<AStaticMeshActor*> NewActors;
+
+	if (!IsValid(InCloner))
+	{
+		return NewActors;
+	}
+
+	UWorld* World = InCloner->GetWorld();
+	AActor* Owner = InCloner->GetOwner();
+
+	if (!IsValid(World)
+		|| !IsValid(Owner)
+#if WITH_EDITOR
+		|| Owner->bIsEditorPreviewActor
+#endif
+		)
+	{
+		return NewActors;
+	}
+
+	FCEClonerMeshBuilder ClonerMeshBuilder;
+
+	if (!ClonerMeshBuilder.AppendComponent(InCloner, InCloner->GetComponentTransform()))
+	{
+		return NewActors;
+	}
+
+	if (ClonerMeshBuilder.GetMeshInstanceCount() == 0)
+	{
+		return NewActors;
+	}
+
+#if WITH_EDITOR
+	UPackage* ClonerPackage = InCloner->GetPackage();
+
+	if (!ClonerPackage)
+	{
+		return NewActors;
+	}
+
+	FString PackagePath;
+	if (!PickAssetPath(ClonerPackage->GetLoadedPath().GetPackageName(), PackagePath))
+	{
+		return NewActors;
+	}
+#endif
+
+	// Spawn attached actor with same flags as this actor
+	FActorSpawnParameters SpawnParameters;
+	SpawnParameters.Owner = Owner;
+	SpawnParameters.ObjectFlags = InCloner->GetFlags();
+#if WITH_EDITOR
+	SpawnParameters.bTemporaryEditorActor = false;
+#endif
+
+	const FVector ClonerLocation = InCloner->GetComponentLocation();
+	const FRotator ClonerRotation = InCloner->GetComponentRotation();
+
+	// Create a Group Actor to hold all actors related to this operation
+	AActor* GroupActor = World->SpawnActor<AActor>(ClonerLocation, ClonerRotation, SpawnParameters);
+
+	if (!GroupActor)
+	{
+		return NewActors;
+	}
+
+	CreateComponent(GroupActor, USceneComponent::StaticClass());
+
+#if WITH_EDITOR
+	FActorLabelUtilities::SetActorLabelUnique(GroupActor, Owner->GetActorNameOrLabel() + TEXT("_SM_Instances"));
+#endif
+
+	NewActors.Reserve(ClonerMeshBuilder.GetMeshInstanceCount());
+
+#if WITH_EDITOR
+	for (uint32 MeshIndex : ClonerMeshBuilder.GetMeshIndexes())
+	{
+		UStaticMesh* StaticMesh = nullptr;
+
+		const FString AssetPath = PackagePath + TEXT("SM_") + Owner->GetActorNameOrLabel() + TEXT("_") + FString::FromInt(MeshIndex);
+		StaticMesh = CreateAssetPackage<UStaticMesh>(AssetPath);
+
+		if (!StaticMesh)
+		{
+			continue;
+		}
+
+		TArray<FCEClonerMeshBuilder::FCEClonerMeshInstanceData> Instances;
+		ClonerMeshBuilder.BuildStaticMesh(MeshIndex, StaticMesh, Instances);
+
+		// Replace material references that are no assets to avoid save issue in new packages
+		for (int32 Index = 0; Index < StaticMesh->GetNumSections(/** LOD */0); Index++)
+		{
+			UMaterialInterface* Material = StaticMesh->GetMaterial(Index);
+
+			if (Material && !Material->IsAsset())
+			{
+				Material = nullptr;
+			}
+
+			StaticMesh->SetMaterial(Index, Material);
+		}
+
+		for (const FCEClonerMeshBuilder::FCEClonerMeshInstanceData& Instance : Instances)
+		{
+			if (AStaticMeshActor* StaticMeshActor = World->SpawnActor<AStaticMeshActor>(ClonerLocation, ClonerRotation, SpawnParameters))
+			{
+				StaticMeshActor->SetMobility(EComponentMobility::Movable);
+
+				UStaticMeshComponent* StaticMeshComponent = StaticMeshActor->GetStaticMeshComponent();
+
+				StaticMeshComponent->SetStaticMesh(StaticMesh);
+				StaticMeshActor->SetActorTransform(Instance.Transform);
+
+				for (int32 MaterialIndex = 0; MaterialIndex < Instance.MeshMaterials.Num(); MaterialIndex++)
+				{
+					StaticMeshComponent->SetMaterial(MaterialIndex, Instance.MeshMaterials[MaterialIndex].Get());
+				}
+
+				StaticMeshActor->AttachToActor(GroupActor, FAttachmentTransformRules::KeepWorldTransform);
+
+				FActorLabelUtilities::SetActorLabelUnique(StaticMeshActor, Owner->GetActorNameOrLabel() + TEXT("_SM_Instance"));
+
+				NewActors.Add(StaticMeshActor);
+			}
+		}
+	}
+
+#else
+
+	for (int32 Index = 0; Index < ClonerMeshBuilder.GetMeshInstanceCount(); Index++)
+	{
+		if (AStaticMeshActor* StaticMeshActor = World->SpawnActor<AStaticMeshActor>(ClonerLocation, ClonerRotation, SpawnParameters))
+		{
+			StaticMeshActor->SetMobility(EComponentMobility::Movable);
+
+			UStaticMeshComponent* StaticMeshComponent = StaticMeshActor->GetStaticMeshComponent();
+
+			UStaticMesh* StaticMesh = NewObject<UStaticMesh>(StaticMeshComponent);
+
+			FCEClonerMeshBuilder::FCEClonerMeshInstanceData MeshData;
+			ClonerMeshBuilder.BuildStaticMesh(Index, StaticMesh, MeshData);
+
+			StaticMeshComponent->SetStaticMesh(StaticMesh);
+			StaticMeshActor->SetActorTransform(MeshData.Transform);
+
+			for (int32 MaterialIndex = 0; MaterialIndex < MeshData.MeshMaterials.Num(); MaterialIndex++)
+			{
+				StaticMeshComponent->SetMaterial(MaterialIndex, MeshData.MeshMaterials[MaterialIndex].Get());
+			}
+
+			StaticMeshActor->AttachToActor(GroupActor, FAttachmentTransformRules::KeepWorldTransform);
+
+			NewActors.Add(StaticMeshActor);
+		}
+	}
+#endif
+
+	return NewActors;
+}
+
+TArray<ADynamicMeshActor*> UE::ClonerEffector::Conversion::ConvertClonerToDynamicMeshes(UCEClonerComponent* InCloner)
+{
+	TArray<ADynamicMeshActor*> NewActors;
+
+	if (!IsValid(InCloner))
+	{
+		return NewActors;
+	}
+
+	UWorld* World = InCloner->GetWorld();
+	AActor* Owner = InCloner->GetOwner();
+
+	if (!IsValid(World)
+		|| !IsValid(Owner)
+#if WITH_EDITOR
+		|| Owner->bIsEditorPreviewActor
+#endif
+		)
+	{
+		return NewActors;
+	}
+
+	FCEClonerMeshBuilder ClonerMeshBuilder;
+
+	if (!ClonerMeshBuilder.AppendComponent(InCloner, InCloner->GetComponentTransform()))
+	{
+		return NewActors;
+	}
+
+	if (ClonerMeshBuilder.GetMeshInstanceCount() == 0)
+	{
+		return NewActors;
+	}
+
+	// Spawn attached actor with same flags as this actor
+	FActorSpawnParameters SpawnParameters;
+	SpawnParameters.Owner = Owner;
+	SpawnParameters.ObjectFlags = InCloner->GetFlags();
+#if WITH_EDITOR
+	SpawnParameters.bTemporaryEditorActor = false;
+#endif
+
+	const FVector ClonerLocation = InCloner->GetComponentLocation();
+	const FRotator ClonerRotation = InCloner->GetComponentRotation();
+
+	// Create a Group Actor to hold all actors related to this operation
+	AActor* GroupActor = World->SpawnActor<AActor>(ClonerLocation, ClonerRotation, SpawnParameters);
+
+	if (!GroupActor)
+	{
+		return NewActors;
+	}
+
+	CreateComponent(GroupActor, USceneComponent::StaticClass());
+
+#if WITH_EDITOR
+	FActorLabelUtilities::SetActorLabelUnique(GroupActor, Owner->GetActorNameOrLabel() + TEXT("_DM_Instances"));
+#endif
+
+	for (int32 Index = 0; Index < ClonerMeshBuilder.GetMeshInstanceCount(); Index++)
+	{
+		if (ADynamicMeshActor* DynamicMeshActor = World->SpawnActor<ADynamicMeshActor>(ClonerLocation, ClonerRotation, SpawnParameters))
+		{
+			UDynamicMeshComponent* DynamicMeshComponent = DynamicMeshActor->GetDynamicMeshComponent();
+
+			FCEClonerMeshBuilder::FCEClonerMeshInstanceData MeshData;
+			ClonerMeshBuilder.BuildDynamicMesh(Index, DynamicMeshComponent->GetDynamicMesh(), MeshData);
+
+			DynamicMeshActor->SetActorTransform(MeshData.Transform);
+
+			for (int32 MaterialIndex = 0; MaterialIndex < MeshData.MeshMaterials.Num(); MaterialIndex++)
+			{
+				DynamicMeshComponent->SetMaterial(MaterialIndex, MeshData.MeshMaterials[MaterialIndex].Get());
+			}
+
+			DynamicMeshActor->AttachToActor(GroupActor, FAttachmentTransformRules::KeepWorldTransform);
+
+#if WITH_EDITOR
+			FActorLabelUtilities::SetActorLabelUnique(DynamicMeshActor, Owner->GetActorNameOrLabel() + TEXT("_DM_Instance"));
+#endif
+
+			NewActors.Add(DynamicMeshActor);
+		}
+	}
+
+	return NewActors;
+}
+
+TArray<AActor*> UE::ClonerEffector::Conversion::ConvertClonerToInstancedStaticMeshes(UCEClonerComponent* InCloner)
+{
+	TArray<AActor*> NewActors;
+
+	if (!IsValid(InCloner))
+	{
+		return NewActors;
+	}
+
+	UWorld* World = InCloner->GetWorld();
+	AActor* Owner = InCloner->GetOwner();
+
+	if (!IsValid(World)
+		|| !IsValid(Owner)
+#if WITH_EDITOR
+		|| Owner->bIsEditorPreviewActor
+#endif
+		)
+	{
+		return NewActors;
+	}
+
+	FCEClonerMeshBuilder ClonerMeshBuilder;
+
+	if (!ClonerMeshBuilder.AppendComponent(InCloner, InCloner->GetComponentTransform()))
+	{
+		return NewActors;
+	}
+
+	if (ClonerMeshBuilder.GetMeshInstanceCount() == 0)
+	{
+		return NewActors;
+	}
+
+#if WITH_EDITOR
+	UPackage* ClonerPackage = InCloner->GetPackage();
+
+	if (!ClonerPackage)
+	{
+		return NewActors;
+	}
+
+	FString PackagePath;
+	if (!PickAssetPath(ClonerPackage->GetLoadedPath().GetPackageName(), PackagePath))
+	{
+		return NewActors;
+	}
+#endif
+
+	// Spawn attached actor with same flags as this actor
+	FActorSpawnParameters SpawnParameters;
+	SpawnParameters.Owner = Owner;
+	SpawnParameters.ObjectFlags = InCloner->GetFlags();
+#if WITH_EDITOR
+	SpawnParameters.bTemporaryEditorActor = false;
+#endif
+
+	const FVector ClonerLocation = InCloner->GetComponentLocation();
+	const FRotator ClonerRotation = InCloner->GetComponentRotation();
+
+	// Create a Group Actor to hold all actors related to this operation
+	AActor* GroupActor = World->SpawnActor<AActor>(ClonerLocation, ClonerRotation, SpawnParameters);
+
+	if (!GroupActor)
+	{
+		return NewActors;
+	}
+
+	CreateComponent(GroupActor, USceneComponent::StaticClass());
+
+#if WITH_EDITOR
+	FActorLabelUtilities::SetActorLabelUnique(GroupActor, Owner->GetActorNameOrLabel() + TEXT("_ISM_Instances"));
+#endif
+
+	for (uint32 MeshIndex : ClonerMeshBuilder.GetMeshIndexes())
+	{
+		if (AActor* ISMActor = World->SpawnActor<AActor>(ClonerLocation, ClonerRotation, SpawnParameters))
+		{
+			UInstancedStaticMeshComponent* ISMComponent = Cast<UInstancedStaticMeshComponent>(CreateComponent(ISMActor, UInstancedStaticMeshComponent::StaticClass()));
+
+			UStaticMesh* StaticMesh = nullptr;
+
+#if WITH_EDITOR
+			const FString AssetPath = PackagePath + TEXT("SM_") + Owner->GetActorNameOrLabel() + TEXT("_") + FString::FromInt(MeshIndex);
+			StaticMesh = CreateAssetPackage<UStaticMesh>(AssetPath);
+#else
+			StaticMesh = NewObject<UStaticMesh>(ISMComponent);
+#endif
+
+			if (!StaticMesh)
+			{
+				continue;
+			}
+
+			TArray<FCEClonerMeshBuilder::FCEClonerMeshInstanceData> Instances;
+			ClonerMeshBuilder.BuildStaticMesh(MeshIndex, StaticMesh, Instances);
+
+#if WITH_EDITOR
+			// Replace material references that are no assets to avoid save issue in new packages
+			for (int32 Index = 0; Index < StaticMesh->GetNumSections(/** LOD */0); Index++)
+			{
+				UMaterialInterface* Material = StaticMesh->GetMaterial(Index);
+
+				if (Material && !Material->IsAsset())
+				{
+					Material = nullptr;
+				}
+
+				StaticMesh->SetMaterial(Index, Material);
+			}
+#endif
+
+			ISMComponent->SetStaticMesh(StaticMesh);
+
+			for (const FCEClonerMeshBuilder::FCEClonerMeshInstanceData& Instance : Instances)
+			{
+				ISMComponent->AddInstance(Instance.Transform, /** WorldSpace */false);
+			}
+
+			if (!Instances.IsEmpty())
+			{
+				constexpr int32 InstanceIndex = 0;
+
+				// ISM component do not have a way to set different materials per instance, just pick first instance materials...
+				for (int32 MaterialIndex = 0; MaterialIndex < Instances[InstanceIndex].MeshMaterials.Num(); MaterialIndex++)
+				{
+					ISMComponent->SetMaterial(MaterialIndex, Instances[InstanceIndex].MeshMaterials[MaterialIndex].Get());
+				}
+			}
+
+			ISMActor->AttachToActor(GroupActor, FAttachmentTransformRules::KeepWorldTransform);
+
+#if WITH_EDITOR
+			FActorLabelUtilities::SetActorLabelUnique(ISMActor, Owner->GetActorNameOrLabel() + TEXT("_ISM_Instance"));
+#endif
+
+			NewActors.Add(ISMActor);
+		}
+	}
+
+	return NewActors;
+}
+
+UActorComponent* UE::ClonerEffector::Conversion::CreateComponent(AActor* InActor, TSubclassOf<USceneComponent> InComponentClass)
+{
+	USceneComponent* const NewComponent = NewObject<USceneComponent>(InActor
+		, InComponentClass.Get()
+		, MakeUniqueObjectName(InActor, InComponentClass.Get(), InComponentClass->GetFName())
+		, RF_Transactional);
+
+	InActor->SetRootComponent(NewComponent);
+	InActor->AddInstanceComponent(NewComponent);
+	NewComponent->OnComponentCreated();
+	NewComponent->RegisterComponent();
+
+#if WITH_EDITOR
+	InActor->RerunConstructionScripts();
+#endif
+
+	return NewComponent;
+}
+
+#if WITH_EDITOR
+bool UE::ClonerEffector::Conversion::PickAssetPath(const FString& InDefaultPath, FString& OutPickedPath)
+{
+	TSharedPtr<SDlgPickPath> DialogWidget = SNew(SDlgPickPath)
+		.Title(LOCTEXT("PickAssetsLocation", "Choose Asset(s) Location"))
+		.DefaultPath(FText::FromString(InDefaultPath));
+
+	if (DialogWidget->ShowModal() != EAppReturnType::Ok)
+	{
+		OutPickedPath = TEXT("");
+		return false;
+	}
+
+	OutPickedPath = DialogWidget->GetPath().ToString() + TEXT("/");
+	return true;
+}
+
+UObject* UE::ClonerEffector::Conversion::CreateAssetPackage(TSubclassOf<UObject> InAssetClass, const FString& InAssetPath)
+{
+	if (!IsValid(InAssetClass.Get()) || InAssetPath.IsEmpty())
+	{
+		return nullptr;
+	}
+
+	FString PackageName = FPackageName::ObjectPathToPackageName(InAssetPath);
+	FAssetToolsModule& AssetToolsModule = FModuleManager::LoadModuleChecked<FAssetToolsModule>("AssetTools");
+
+	FString AssetName;
+	AssetToolsModule.Get().CreateUniqueAssetName(PackageName, TEXT(""), PackageName, AssetName);
+
+	FString UserPackageName = InAssetPath;
+	FName UserAssetName(*FPackageName::GetLongPackageAssetName(UserPackageName));
+
+	// is input name valid ?
+	if (UserAssetName.IsNone())
+	{
+		// Use default if invalid
+		UserPackageName = PackageName;
+		UserAssetName = *AssetName;
+	}
+
+	// Find/create package
+	UPackage* Package = CreatePackage(*UserPackageName);
+
+	if (!Package)
+	{
+		return nullptr;
+	}
+
+	// Create asset object
+	UObject* AssetObject = NewObject<UObject>(Package, InAssetClass.Get(), UserAssetName, RF_Public | RF_Standalone);
+
+	// Notify asset registry of new asset
+	FAssetRegistryModule::AssetCreated(AssetObject);
+
+	return AssetObject;
+}
+#endif
+
+#undef LOCTEXT_NAMESPACE
diff --git a/Engine/Plugins/Experimental/ClonerEffector/Source/ClonerEffector/Private/Cloner/CEClonerActor.cpp b/Engine/Plugins/Experimental/ClonerEffector/Source/ClonerEffector/Private/Cloner/CEClonerActor.cpp
--- a/Engine/Plugins/Experimental/ClonerEffector/Source/ClonerEffector/Private/Cloner/CEClonerActor.cpp
+++ b/Engine/Plugins/Experimental/ClonerEffector/Source/ClonerEffector/Private/Cloner/CEClonerActor.cpp
@@ -21,6 +21,7 @@
 #include "Cloner/Layouts/CEClonerSphereRandomLayout.h"
 #include "Cloner/Layouts/CEClonerSphereUniformLayout.h"
 #include "Cloner/Layouts/CEClonerSplineLayout.h"
+#include "Containers/Ticker.h"
 #include "Effector/CEEffectorActor.h"
 #include "Materials/MaterialInterface.h"
 #include "UObject/ConstructorHelpers.h"
@@ -55,7 +56,7 @@ FCustomVersionRegistration GRegisterCEClonerActorVersion(FCEClonerActorVersion::
 ACEClonerActor::ACEClonerActor()
 {
 	SetCanBeDamaged(false);
-	PrimaryActorTick.bCanEverTick          = false;
+	PrimaryActorTick.bCanEverTick = false;
 
 	ClonerComponent = CreateDefaultSubobject<UCEClonerComponent>(TEXT("AvaClonerComponent"));
 	SetRootComponent(ClonerComponent);
@@ -113,11 +114,8 @@ void ACEClonerActor::PostActorCreated()
 void ACEClonerActor::OnClonerInitialized(UCEClonerComponent* InClonerComponent)
 {
 #if WITH_EDITOR
-	if (bSpawnDefaultActorAttached)
-	{
-		bSpawnDefaultActorAttached = false;
-		ClonerComponent->CreateDefaultActorAttached();
-	}
+	// PostActorCreated is sometimes called after the cloner is initialized
+	FTSTicker::GetCoreTicker().AddTicker(FTickerDelegate::CreateUObject(this, &ACEClonerActor::SpawnDefaultActorAttached));
 #endif
 }
 
@@ -128,7 +126,7 @@ void ACEClonerActor::MigrateDeprecatedProperties()
 		return;
 	}
 
-	UE_LOG(LogCEClonerActor, Warning, TEXT("%s : Cloner migrating from version %i to latest %i"), *GetActorNameOrLabel(), MigrateToVersion, FCEClonerActorVersion::LatestVersion);
+	UE_LOG(LogCEClonerActor, Warning, TEXT("%s : Cloner migrating from version %i to latest %i, please re-save this asset"), *GetActorNameOrLabel(), MigrateToVersion, FCEClonerActorVersion::LatestVersion);
 
 	PRAGMA_DISABLE_DEPRECATION_WARNINGS
 
@@ -137,6 +135,7 @@ void ACEClonerActor::MigrateDeprecatedProperties()
 		ClonerComponent->SetEnabled(bEnabled);
 		ClonerComponent->SetTreeUpdateInterval(TreeUpdateInterval);
 		ClonerComponent->SetSeed(Seed);
+		ClonerComponent->SetColor(Color);
 		ClonerComponent->SetLayoutName(LayoutName);
 
 #if WITH_EDITOR
@@ -214,8 +213,6 @@ void ACEClonerActor::MigrateDeprecatedProperties()
 
 		if (UCEClonerEffectorExtension* EffectorExtension = ClonerComponent->FindOrAddExtension<UCEClonerEffectorExtension>())
 		{
-			EffectorExtension->SetColor(Color);
-
 			for (const TWeakObjectPtr<ACEEffectorActor>& EffectorWeak : EffectorsWeak)
 			{
 				if (ACEEffectorActor* Effector = EffectorWeak.Get())
@@ -356,6 +353,17 @@ void ACEClonerActor::MigrateDeprecatedProperties()
 }
 
 #if WITH_EDITOR
+bool ACEClonerActor::SpawnDefaultActorAttached(float)
+{
+	if (bSpawnDefaultActorAttached)
+	{
+		bSpawnDefaultActorAttached = false;
+		ClonerComponent->CreateDefaultActorAttached();
+	}
+
+	return false;
+}
+
 void ACEClonerActor::OnEditorSelectionChanged(UObject* InSelection)
 {
 	if (const USelection* ActorSelection = Cast<USelection>(InSelection))
diff --git a/Engine/Plugins/Experimental/ClonerEffector/Source/ClonerEffector/Private/Cloner/CEClonerComponent.cpp b/Engine/Plugins/Experimental/ClonerEffector/Source/ClonerEffector/Private/Cloner/CEClonerComponent.cpp
--- a/Engine/Plugins/Experimental/ClonerEffector/Source/ClonerEffector/Private/Cloner/CEClonerComponent.cpp
+++ b/Engine/Plugins/Experimental/ClonerEffector/Source/ClonerEffector/Private/Cloner/CEClonerComponent.cpp
@@ -27,6 +27,7 @@
 #include "Framework/Notifications/NotificationManager.h"
 #include "Materials/Material.h"
 #include "Misc/MessageDialog.h"
+#include "Misc/ScopedSlowTask.h"
 #include "Widgets/Notifications/SNotificationList.h"
 #endif
 
@@ -37,11 +38,6 @@ DEFINE_LOG_CATEGORY_STATIC(LogCEClonerComponent, Log, All);
 UCEClonerComponent::UCEClonerComponent()
 	: UNiagaraComponent()
 {
-	PrimaryComponentTick.bCanEverTick          = true;
-	PrimaryComponentTick.bStartWithTickEnabled = true;
-	PrimaryComponentTick.bTickEvenWhenPaused   = true;
-	PrimaryComponentTick.bHighPriority         = true;
-
 	CastShadow = true;
 	bReceivesDecals = true;
 	bAutoActivate = true;
@@ -107,8 +103,6 @@ void UCEClonerComponent::PostInitProperties()
 
 void UCEClonerComponent::PostLoad()
 {
-	SetAsset(nullptr);
-
 	Super::PostLoad();
 
 	InitializeCloner();
@@ -119,6 +113,9 @@ void UCEClonerComponent::PostEditUndo()
 {
 	Super::PostEditUndo();
 
+	// Reregister ticker in case this object was destroyed then undo
+	RegisterTicker();
+
 	ForceUpdateCloner();
 }
 
@@ -126,6 +123,7 @@ const TCEPropertyChangeDispatcher<UCEClonerComponent> UCEClonerComponent::Proper
 {
 	{ GET_MEMBER_NAME_CHECKED(UCEClonerComponent, bEnabled), &UCEClonerComponent::OnEnabledChanged },
 	{ GET_MEMBER_NAME_CHECKED(UCEClonerComponent, Seed), &UCEClonerComponent::OnSeedChanged },
+	{ GET_MEMBER_NAME_CHECKED(UCEClonerComponent, Color), &UCEClonerComponent::OnColorChanged },
 	/** Layout */
 	{ GET_MEMBER_NAME_CHECKED(UCEClonerComponent, LayoutName), &UCEClonerComponent::OnLayoutNameChanged },
 	{ GET_MEMBER_NAME_CHECKED(UCEClonerComponent, bVisualizerSpriteVisible), &UCEClonerComponent::OnVisualizerSpriteVisibleChanged },
@@ -139,6 +137,13 @@ void UCEClonerComponent::PostEditChangeProperty(FPropertyChangedEvent& InPropert
 }
 #endif
 
+void UCEClonerComponent::OnComponentCreated()
+{
+	Super::OnComponentCreated();
+
+	InitializeCloner();
+}
+
 void UCEClonerComponent::UpdateClonerRenderState()
 {
 	/**
@@ -807,6 +812,7 @@ void UCEClonerComponent::UpdateActorBakedDynamicMesh(AActor* InActor)
 
 	MeshBuilder.AppendActor(InActor);
 	MeshBuilder.BuildDynamicMesh(Mesh, MeshMaterials);
+	MeshBuilder.Reset();
 
 	AttachmentItem->BakedMesh = Mesh;
 	AttachmentItem->BakedMaterials = MoveTemp(MeshMaterials);
@@ -904,6 +910,7 @@ void UCEClonerComponent::UpdateRootActorBakedStaticMesh(AActor* InRootActor)
 	UStaticMesh* Mesh = NewObject<UStaticMesh>();
 	TArray<TWeakObjectPtr<UMaterialInterface>> MeshMaterials;
 	bClonerMeshesDirty = MeshBuilder.BuildStaticMesh(Mesh, MeshMaterials);
+	MeshBuilder.Reset();
 
 	ClonerTree.MergedBakedMeshes[RootIdx] = Mesh;
 
@@ -996,13 +1003,6 @@ void UCEClonerComponent::ShowMaterialWarning(int32 InMaterialCount)
 }
 #endif // WITH_EDITOR
 
-void UCEClonerComponent::TickComponent(float InDeltaTime, ELevelTick InTickType, FActorComponentTickFunction* InThisTickFunction)
-{
-	TickCloner(InDeltaTime);
-
-	Super::TickComponent(InDeltaTime, InTickType, InThisTickFunction);
-}
-
 void UCEClonerComponent::InitializeCloner()
 {
 	if (bClonerInitialized)
@@ -1012,8 +1012,7 @@ void UCEClonerComponent::InitializeCloner()
 
 	bClonerInitialized = true;
 
-	/** Allow tick in editor without hitting check */
-	SetForceSolo(true);
+	SetAsset(nullptr);
 
 	OnLayoutNameChanged();
 
@@ -1022,13 +1021,22 @@ void UCEClonerComponent::InitializeCloner()
 #endif
 
 	OnClonerInitializedDelegate.Broadcast(this);
+
+	/** Register a custom ticker to avoid using the component tick that needs the simulation to be solo */
+	RegisterTicker();
+}
+
+void UCEClonerComponent::RegisterTicker()
+{
+	FTSTicker::GetCoreTicker().RemoveTicker(ClonerTickerHandle);
+	ClonerTickerHandle = FTSTicker::GetCoreTicker().AddTicker(FTickerDelegate::CreateUObject(this, &UCEClonerComponent::TickCloner));
 }
 
-void UCEClonerComponent::TickCloner(float InDelta)
+bool UCEClonerComponent::TickCloner(float InDelta)
 {
 	if (!bClonerInitialized)
 	{
-		InitializeCloner();
+		return false;
 	}
 
 	if (bEnabled)
@@ -1049,6 +1057,8 @@ void UCEClonerComponent::TickCloner(float InDelta)
 			RequestClonerUpdate(true);
 		}
 	}
+
+	return true;
 }
 
 void UCEClonerComponent::SetEnabled(bool bInEnable)
@@ -1083,6 +1093,17 @@ void UCEClonerComponent::SetSeed(int32 InSeed)
 	OnSeedChanged();
 }
 
+void UCEClonerComponent::SetColor(const FLinearColor& InColor)
+{
+	if (InColor.Equals(Color))
+	{
+		return;
+	}
+
+	Color = InColor;
+	OnColorChanged();
+}
+
 void UCEClonerComponent::SetLayoutName(FName InLayoutName)
 {
 	if (LayoutName == InLayoutName)
@@ -1150,7 +1171,10 @@ void UCEClonerComponent::CreateDefaultActorAttached()
 	// Only spawn if world is valid and not a preview actor
 	UWorld* World = GetWorld();
 	AActor* Owner = GetOwner();
-	if (!World || !Owner || Owner->bIsEditorPreviewActor)
+
+	if (!IsValid(World)
+		|| !IsValid(Owner)
+		|| Owner->bIsEditorPreviewActor)
 	{
 		return;
 	}
@@ -1195,6 +1219,116 @@ void UCEClonerComponent::CreateDefaultActorAttached()
 		FActorLabelUtilities::SetActorLabelUnique(DefaultActorAttached, TEXT("DefaultClone"));
 	}
 }
+
+void UCEClonerComponent::ConvertToStaticMesh()
+{
+	if (!IsValid(this) || !bEnabled)
+	{
+		return;
+	}
+
+	FScopedSlowTask SlowTask(0.0f, LOCTEXT("ConvertToStaticMesh", "Converting cloner to static mesh"));
+	SlowTask.MakeDialog();
+
+	UE_LOG(LogCEClonerComponent, Log, TEXT("%s : Request ConvertToStaticMesh..."), *GetOwner()->GetActorNameOrLabel())
+
+	if (UE::ClonerEffector::Conversion::ConvertClonerToStaticMesh(this))
+	{
+		UE_LOG(LogCEClonerComponent, Log, TEXT("%s : ConvertToStaticMesh Completed"), *GetOwner()->GetActorNameOrLabel())
+	}
+	else
+	{
+		UE_LOG(LogCEClonerComponent, Warning, TEXT("%s : ConvertToStaticMesh Failed"), *GetOwner()->GetActorNameOrLabel())
+	}
+}
+
+void UCEClonerComponent::ConvertToDynamicMesh()
+{
+	if (!IsValid(this) || !bEnabled)
+	{
+		return;
+	}
+
+	FScopedSlowTask SlowTask(0.0f, LOCTEXT("ConvertToDynamicMesh", "Converting cloner to dynamic mesh"));
+	SlowTask.MakeDialog();
+
+	UE_LOG(LogCEClonerComponent, Log, TEXT("%s : Request ConvertToDynamicMesh..."), *GetOwner()->GetActorNameOrLabel())
+
+	if (UE::ClonerEffector::Conversion::ConvertClonerToDynamicMesh(this))
+	{
+		UE_LOG(LogCEClonerComponent, Log, TEXT("%s : ConvertToDynamicMesh Completed"), *GetOwner()->GetActorNameOrLabel())
+	}
+	else
+	{
+		UE_LOG(LogCEClonerComponent, Warning, TEXT("%s : ConvertToDynamicMesh Failed"), *GetOwner()->GetActorNameOrLabel())
+	}
+}
+
+void UCEClonerComponent::ConvertToStaticMeshes()
+{
+	if (!IsValid(this) || !bEnabled)
+	{
+		return;
+	}
+
+	FScopedSlowTask SlowTask(0.0f, LOCTEXT("ConvertToStaticMeshes", "Converting cloner to static meshes"));
+	SlowTask.MakeDialog();
+
+	UE_LOG(LogCEClonerComponent, Log, TEXT("%s : Request ConvertToStaticMeshes..."), *GetOwner()->GetActorNameOrLabel())
+
+	if (!UE::ClonerEffector::Conversion::ConvertClonerToStaticMeshes(this).IsEmpty())
+	{
+		UE_LOG(LogCEClonerComponent, Log, TEXT("%s : ConvertToStaticMeshes Completed"), *GetOwner()->GetActorNameOrLabel())
+	}
+	else
+	{
+		UE_LOG(LogCEClonerComponent, Warning, TEXT("%s : ConvertToStaticMeshes Failed"), *GetOwner()->GetActorNameOrLabel())
+	}
+}
+
+void UCEClonerComponent::ConvertToDynamicMeshes()
+{
+	if (!IsValid(this) || !bEnabled)
+	{
+		return;
+	}
+
+	FScopedSlowTask SlowTask(0.0f, LOCTEXT("ConvertToDynamicMeshes", "Converting cloner to dynamic meshes"));
+	SlowTask.MakeDialog();
+
+	UE_LOG(LogCEClonerComponent, Log, TEXT("%s : Request ConvertToDynamicMeshes..."), *GetOwner()->GetActorNameOrLabel())
+
+	if (!UE::ClonerEffector::Conversion::ConvertClonerToDynamicMeshes(this).IsEmpty())
+	{
+		UE_LOG(LogCEClonerComponent, Log, TEXT("%s : ConvertToDynamicMeshes Completed"), *GetOwner()->GetActorNameOrLabel())
+	}
+	else
+	{
+		UE_LOG(LogCEClonerComponent, Warning, TEXT("%s : ConvertToDynamicMeshes Failed"), *GetOwner()->GetActorNameOrLabel())
+	}
+}
+
+void UCEClonerComponent::ConvertToInstancedStaticMeshes()
+{
+	if (!IsValid(this) || !bEnabled)
+	{
+		return;
+	}
+
+	FScopedSlowTask SlowTask(0.0f, LOCTEXT("ConvertToInstancedStaticMeshes", "Converting cloner to instanced static meshes"));
+	SlowTask.MakeDialog();
+
+	UE_LOG(LogCEClonerComponent, Log, TEXT("%s : Request ConvertToInstancedStaticMeshes..."), *GetOwner()->GetActorNameOrLabel())
+
+	if (!UE::ClonerEffector::Conversion::ConvertClonerToInstancedStaticMeshes(this).IsEmpty())
+	{
+		UE_LOG(LogCEClonerComponent, Log, TEXT("%s : ConvertToInstancedStaticMeshes Completed"), *GetOwner()->GetActorNameOrLabel())
+	}
+	else
+	{
+		UE_LOG(LogCEClonerComponent, Warning, TEXT("%s : ConvertToInstancedStaticMeshes Failed"), *GetOwner()->GetActorNameOrLabel())
+	}
+}
 #endif
 
 void UCEClonerComponent::RequestClonerUpdate(bool bInImmediate)
@@ -1248,6 +1382,21 @@ void UCEClonerComponent::OnClonerDisabled()
 	SetAsset(nullptr);
 }
 
+void UCEClonerComponent::OnClonerSetEnabled(const UWorld* InWorld, bool bInEnabled, bool bInTransact)
+{
+	if (GetWorld() == InWorld)
+	{
+#if WITH_EDITOR
+		if (bInTransact)
+		{
+			Modify();
+		}
+#endif
+
+		SetEnabled(bInEnabled);
+	}
+}
+
 void UCEClonerComponent::OnSeedChanged()
 {
 	if (!bEnabled)
@@ -1260,6 +1409,11 @@ void UCEClonerComponent::OnSeedChanged()
 	RequestClonerUpdate();
 }
 
+void UCEClonerComponent::OnColorChanged()
+{
+	SetColorParameter(TEXT("EffectorDefaultColor"), Color);
+}
+
 void UCEClonerComponent::OnLayoutNameChanged()
 {
 	if (!bEnabled)
@@ -1510,6 +1664,7 @@ void UCEClonerComponent::OnActiveLayoutChanged()
 	}
 
 	OnSeedChanged();
+	OnColorChanged();
 
 	Layout->UpdateLayoutParameters();
 
diff --git a/Engine/Plugins/Experimental/ClonerEffector/Source/ClonerEffector/Private/Cloner/CEClonerMeshBuilder.cpp b/Engine/Plugins/Experimental/ClonerEffector/Source/ClonerEffector/Private/Cloner/CEClonerMeshBuilder.cpp
--- a/Engine/Plugins/Experimental/ClonerEffector/Source/ClonerEffector/Private/Cloner/CEClonerMeshBuilder.cpp
+++ b/Engine/Plugins/Experimental/ClonerEffector/Source/ClonerEffector/Private/Cloner/CEClonerMeshBuilder.cpp
@@ -17,20 +17,34 @@
 #include "GeometryScript/MeshBasicEditFunctions.h"
 #include "GeometryScript/SceneUtilityFunctions.h"
 #include "Materials/MaterialInterface.h"
+#include "NiagaraComponent.h"
+#include "NiagaraEmitter.h"
+#include "NiagaraMeshRendererProperties.h"
+#include "NiagaraSimCacheFunctionLibrary.h"
+#include "NiagaraSystem.h"
 #include "ProceduralMeshComponent.h"
 #include "StaticMeshOperations.h"
 #include "UDynamicMesh.h"
 
+const FCEClonerMeshBuilder::FCEClonerMeshBuilderParams FCEClonerMeshBuilder::DefaultParams;
+
 FCEClonerMeshBuilder::FCEClonerMeshBuilder()
 {
 	OutputDynamicMesh = NewObject<UDynamicMesh>();
 }
 
+TArray<uint32> FCEClonerMeshBuilder::GetMeshIndexes() const
+{
+	TArray<uint32> MeshIndexes;
+	Meshes.GenerateKeyArray(MeshIndexes);
+	return MeshIndexes;
+}
+
 void FCEClonerMeshBuilder::Reset()
 {
-	OutputDynamicMesh->EditMesh([](FDynamicMesh3& EditMesh){ EditMesh.Clear(); });
-	ConvertedMeshes.Empty();
-	OutputMeshMaterials.Empty();
+	ClearOutputMesh();
+	Meshes.Empty();
+	MeshInstances.Empty();
 }
 
 int32 FCEClonerMeshBuilder::AppendActor(const AActor* InActor)
@@ -103,6 +117,13 @@ int32 FCEClonerMeshBuilder::AppendActor(const AActor* InActor)
 				ComponentConverted++;
 			}
 		}
+		else if (UNiagaraComponent* NiagaraComponent = Cast<UNiagaraComponent>(PrimitiveComponent))
+		{
+			if (AppendComponent(NiagaraComponent, SourceTransform))
+			{
+				ComponentConverted++;
+			}
+		}
 	}
 
 	return ComponentConverted;
@@ -110,23 +131,60 @@ int32 FCEClonerMeshBuilder::AppendActor(const AActor* InActor)
 
 bool FCEClonerMeshBuilder::AppendMesh(const UDynamicMesh* InMesh, const TArray<TWeakObjectPtr<UMaterialInterface>>& InMaterials, const FTransform& InTransform)
 {
-	if (!IsValid(InMesh))
+	if (!IsValid(InMesh) || InMesh->GetTriangleCount() == 0)
 	{
 		return false;
 	}
 
-	// Create a copy of the mesh
-	InMesh->ProcessMesh([this, &InTransform](const FDynamicMesh3& EditMesh)
+	return !!AddMeshInstance(InMesh->GetUniqueID(), InTransform, InMaterials, [&InMesh](FDynamicMesh3& InCreateMesh)->bool
 	{
-		FDynamicMesh3 CopyMesh = EditMesh;
-		MeshTransforms::ApplyTransform(CopyMesh, InTransform);
+		// Create a copy of the mesh
+		InMesh->ProcessMesh([&InCreateMesh](const FDynamicMesh3& EditMesh)
+		{
+			InCreateMesh = EditMesh;
+		});
 
-		ConvertedMeshes.Add(MoveTemp(CopyMesh));
+		return true;
 	});
+}
 
-	OutputMeshMaterials.Append(InMaterials);
+bool FCEClonerMeshBuilder::AppendMesh(UStaticMesh* InMesh, const TArray<TWeakObjectPtr<UMaterialInterface>>& InMaterials, const FTransform& InSourceTransform)
+{
+	if (!IsValid(InMesh) || InMesh->GetNumTriangles(/** LOD*/ 0) == 0)
+	{
+		return false;
+	}
 
-	return true;
+	return !!AddMeshInstance(InMesh->GetUniqueID(), InSourceTransform, InMaterials, [this, &InMesh](FDynamicMesh3& InCreateMesh)->bool
+	{
+		// convert to dynamic mesh
+		FGeometryScriptMeshReadLOD StaticMeshLOD;
+		StaticMeshLOD.LODType = EGeometryScriptLODType::RenderData;
+
+		FGeometryScriptCopyMeshFromAssetOptions OutputMeshOptions;
+		OutputMeshOptions.bIgnoreRemoveDegenerates = false;
+		OutputMeshOptions.bRequestTangents = false;
+		OutputMeshOptions.bApplyBuildSettings = false;
+
+		EGeometryScriptOutcomePins OutResult;
+		UGeometryScriptLibrary_StaticMeshFunctions::CopyMeshFromStaticMesh(InMesh, OutputDynamicMesh, OutputMeshOptions, StaticMeshLOD, OutResult);
+
+		if (OutResult != EGeometryScriptOutcomePins::Success)
+		{
+			return false;
+		}
+
+		OutputDynamicMesh->EditMesh([this, &InCreateMesh](FDynamicMesh3& EditMesh)
+		{
+			InCreateMesh = MoveTemp(EditMesh);
+
+			// replace by empty mesh
+			FDynamicMesh3 EmptyMesh;
+			EditMesh = MoveTemp(EmptyMesh);
+		}, EDynamicMeshChangeType::GeneralEdit, EDynamicMeshAttributeChangeFlags::Unknown, true);
+
+		return true;
+	});
 }
 
 bool FCEClonerMeshBuilder::AppendComponent(const UStaticMeshComponent* InComponent, const FTransform& InSourceTransform)
@@ -143,36 +201,15 @@ bool FCEClonerMeshBuilder::AppendComponent(const UStaticMeshComponent* InCompone
 		return false;
 	}
 
-	// convert to dynamic mesh
-	FGeometryScriptMeshReadLOD StaticMeshLOD;
-	StaticMeshLOD.LODType = EGeometryScriptLODType::RenderData;
-
-	FGeometryScriptCopyMeshFromAssetOptions OutputMeshOptions;
-	OutputMeshOptions.bIgnoreRemoveDegenerates = false;
-	OutputMeshOptions.bRequestTangents = false;
-	OutputMeshOptions.bApplyBuildSettings = false;
-
-	EGeometryScriptOutcomePins OutResult;
-	UGeometryScriptLibrary_StaticMeshFunctions::CopyMeshFromStaticMesh(StaticMesh, OutputDynamicMesh, OutputMeshOptions, StaticMeshLOD, OutResult);
-
-	if (OutResult != EGeometryScriptOutcomePins::Success)
-	{
-		return false;
-	}
-
-	// Transform the new mesh relative to the component
 	const FTransform RelativeTransform = InComponent->GetComponentTransform().GetRelativeTransform(InSourceTransform);
-	OutputMeshMaterials.Append(InComponent->GetMaterials());
-	OutputDynamicMesh->EditMesh([this, RelativeTransform](FDynamicMesh3& EditMesh)
+
+	TArray<TWeakObjectPtr<UMaterialInterface>> Materials;
+	Algo::Transform(InComponent->GetMaterials(), Materials, [](UMaterialInterface* InMaterial)
 	{
-		MeshTransforms::ApplyTransform(EditMesh, RelativeTransform);
-		ConvertedMeshes.Add(MoveTemp(EditMesh));
-		// replace by empty mesh
-		FDynamicMesh3 EmptyMesh;
-		EditMesh = MoveTemp(EmptyMesh);
-	}, EDynamicMeshChangeType::GeneralEdit, EDynamicMeshAttributeChangeFlags::Unknown, true);
+		return InMaterial;
+	});
 
-	return true;
+	return AppendMesh(StaticMesh, Materials, RelativeTransform);
 }
 
 bool FCEClonerMeshBuilder::AppendComponent(UProceduralMeshComponent* InComponent, const FTransform& InSourceTransform)
@@ -192,104 +229,106 @@ bool FCEClonerMeshBuilder::AppendComponent(UProceduralMeshComponent* InComponent
 	// Transform the new mesh relative to the component
 	const FTransform RelativeTransform = InComponent->GetComponentTransform().GetRelativeTransform(InSourceTransform);
 
-	FDynamicMesh3 Mesh;
-	Mesh.EnableAttributes();
-	Mesh.Attributes()->EnablePrimaryColors();
-	Mesh.Attributes()->EnableMaterialID();
-	Mesh.Attributes()->SetNumNormalLayers(1);
-	Mesh.Attributes()->SetNumUVLayers(1);
-	Mesh.Attributes()->SetNumPolygroupLayers(1);
-	Mesh.Attributes()->EnableTangents();
-
-	using namespace UE::Geometry;
+	TArray<TWeakObjectPtr<UMaterialInterface>> Materials;
+	Algo::Transform(InComponent->GetMaterials(), Materials, [](UMaterialInterface* InMaterial)
+	{
+		return InMaterial;
+	});
 
-	FDynamicMeshColorOverlay* ColorOverlay = Mesh.Attributes()->PrimaryColors();
-	FDynamicMeshNormalOverlay* NormalOverlay = Mesh.Attributes()->PrimaryNormals();
-	FDynamicMeshUVOverlay* UVOverlay = Mesh.Attributes()->PrimaryUV();
-	FDynamicMeshMaterialAttribute* MaterialAttr = Mesh.Attributes()->GetMaterialID();
-	FDynamicMeshPolygroupAttribute* GroupAttr = Mesh.Attributes()->GetPolygroupLayer(0);
-	FDynamicMeshNormalOverlay* TangentOverlay = Mesh.Attributes()->PrimaryTangents();
+	const uint32 MeshIndex = InComponent->GetUniqueID();
 
-	for (int32 SectionIdx = 0; SectionIdx < SectionCount; SectionIdx++)
+	return !!AddMeshInstance(MeshIndex, RelativeTransform, Materials, [this, &InComponent, &SectionCount](FDynamicMesh3& InCreateMesh)->bool
 	{
-		if (FProcMeshSection* Section = InComponent->GetProcMeshSection(SectionIdx))
+		InCreateMesh.EnableAttributes();
+		InCreateMesh.Attributes()->EnablePrimaryColors();
+		InCreateMesh.Attributes()->EnableMaterialID();
+		InCreateMesh.Attributes()->SetNumNormalLayers(1);
+		InCreateMesh.Attributes()->SetNumUVLayers(1);
+		InCreateMesh.Attributes()->SetNumPolygroupLayers(1);
+		InCreateMesh.Attributes()->EnableTangents();
+
+		using namespace UE::Geometry;
+
+		FDynamicMeshColorOverlay* ColorOverlay = InCreateMesh.Attributes()->PrimaryColors();
+		FDynamicMeshNormalOverlay* NormalOverlay = InCreateMesh.Attributes()->PrimaryNormals();
+		FDynamicMeshUVOverlay* UVOverlay = InCreateMesh.Attributes()->PrimaryUV();
+		FDynamicMeshMaterialAttribute* MaterialAttr = InCreateMesh.Attributes()->GetMaterialID();
+		FDynamicMeshPolygroupAttribute* GroupAttr = InCreateMesh.Attributes()->GetPolygroupLayer(0);
+		FDynamicMeshNormalOverlay* TangentOverlay = InCreateMesh.Attributes()->PrimaryTangents();
+
+		for (int32 SectionIdx = 0; SectionIdx < SectionCount; SectionIdx++)
 		{
-			if (Section->bSectionVisible)
+			if (FProcMeshSection* Section = InComponent->GetProcMeshSection(SectionIdx))
 			{
-				TArray<int32> VtxIds;
-				TArray<int32> NormalIds;
-				TArray<int32> ColorIds;
-				TArray<int32> UVIds;
-				TArray<int32> TaIds;
-
-				// copy vertices data (position, normal, color, UV, tangent)
-				for (FProcMeshVertex& SectionVertex : Section->ProcVertexBuffer)
+				if (Section->bSectionVisible)
 				{
-					int32 VId = Mesh.AppendVertex(SectionVertex.Position);
-					VtxIds.Add(VId);
-
-					int32 NId = NormalOverlay->AppendElement(static_cast<FVector3f>(SectionVertex.Normal));
-					NormalIds.Add(NId);
-
-					int32 CId = ColorOverlay->AppendElement(static_cast<FVector4f>(SectionVertex.Color));
-					ColorIds.Add(CId);
+					TArray<int32> VtxIds;
+					TArray<int32> NormalIds;
+					TArray<int32> ColorIds;
+					TArray<int32> UVIds;
+					TArray<int32> TaIds;
+
+					// copy vertices data (position, normal, color, UV, tangent)
+					for (FProcMeshVertex& SectionVertex : Section->ProcVertexBuffer)
+					{
+						int32 VId = InCreateMesh.AppendVertex(SectionVertex.Position);
+						VtxIds.Add(VId);
 
-					int32 UVId = UVOverlay->AppendElement(static_cast<FVector2f>(SectionVertex.UV0));
-					UVIds.Add(UVId);
+						int32 NId = NormalOverlay->AppendElement(static_cast<FVector3f>(SectionVertex.Normal));
+						NormalIds.Add(NId);
 
-					int32 TaId = TangentOverlay->AppendElement(static_cast<FVector3f>(SectionVertex.Tangent.TangentX));
-					TaIds.Add(TaId);
-				}
+						int32 CId = ColorOverlay->AppendElement(static_cast<FVector4f>(SectionVertex.Color));
+						ColorIds.Add(CId);
 
-				// copy tris data
-				if (Section->ProcIndexBuffer.Num() % 3 != 0)
-				{
-					continue;
-				}
-				for (int32 Idx = 0; Idx < Section->ProcIndexBuffer.Num(); Idx+=3)
-				{
-					int32 VIdx1 = Section->ProcIndexBuffer[Idx];
-					int32 VIdx2 = Section->ProcIndexBuffer[Idx + 1];
-					int32 VIdx3 = Section->ProcIndexBuffer[Idx + 2];
+						int32 UVId = UVOverlay->AppendElement(static_cast<FVector2f>(SectionVertex.UV0));
+						UVIds.Add(UVId);
 
-					int32 VId1 = VtxIds[VIdx1];
-					int32 VId2 = VtxIds[VIdx2];
-					int32 VId3 = VtxIds[VIdx3];
+						int32 TaId = TangentOverlay->AppendElement(static_cast<FVector3f>(SectionVertex.Tangent.TangentX));
+						TaIds.Add(TaId);
+					}
 
-					int32 TId = Mesh.AppendTriangle(VId1, VId2, VId3, SectionIdx);
-					if (TId < 0)
+					// copy tris data
+					if (Section->ProcIndexBuffer.Num() % 3 != 0)
 					{
 						continue;
 					}
 
-					NormalOverlay->SetTriangle(TId, FIndex3i(NormalIds[VIdx1], NormalIds[VIdx2], NormalIds[VIdx3]), true);
-					ColorOverlay->SetTriangle(TId, FIndex3i(ColorIds[VIdx1], ColorIds[VIdx2], ColorIds[VIdx3]), true);
-					UVOverlay->SetTriangle(TId, FIndex3i(UVIds[VIdx1], UVIds[VIdx2], UVIds[VIdx3]), true);
-					TangentOverlay->SetTriangle(TId, FIndex3i(TaIds[VIdx1], TaIds[VIdx2], TaIds[VIdx3]), true);
+					for (int32 Idx = 0; Idx < Section->ProcIndexBuffer.Num(); Idx+=3)
+					{
+						int32 VIdx1 = Section->ProcIndexBuffer[Idx];
+						int32 VIdx2 = Section->ProcIndexBuffer[Idx + 1];
+						int32 VIdx3 = Section->ProcIndexBuffer[Idx + 2];
 
-					MaterialAttr->SetValue(TId, SectionIdx);
-					GroupAttr->SetValue(TId, SectionIdx);
-				}
-			}
-		}
-	}
+						int32 VId1 = VtxIds[VIdx1];
+						int32 VId2 = VtxIds[VIdx2];
+						int32 VId3 = VtxIds[VIdx3];
 
-	if (Mesh.TriangleCount() == 0)
-	{
-		return false;
-	}
+						int32 TId = InCreateMesh.AppendTriangle(VId1, VId2, VId3, SectionIdx);
 
-	MeshTransforms::ApplyTransform(Mesh, RelativeTransform);
+						if (TId < 0)
+						{
+							continue;
+						}
 
-	OutputMeshMaterials.Append(InComponent->GetMaterials());
-	ConvertedMeshes.Add(MoveTemp(Mesh));
+						NormalOverlay->SetTriangle(TId, FIndex3i(NormalIds[VIdx1], NormalIds[VIdx2], NormalIds[VIdx3]), true);
+						ColorOverlay->SetTriangle(TId, FIndex3i(ColorIds[VIdx1], ColorIds[VIdx2], ColorIds[VIdx3]), true);
+						UVOverlay->SetTriangle(TId, FIndex3i(UVIds[VIdx1], UVIds[VIdx2], UVIds[VIdx3]), true);
+						TangentOverlay->SetTriangle(TId, FIndex3i(TaIds[VIdx1], TaIds[VIdx2], TaIds[VIdx3]), true);
 
-	return true;
+						MaterialAttr->SetValue(TId, SectionIdx);
+						GroupAttr->SetValue(TId, SectionIdx);
+					}
+				}
+			}
+		}
+
+		return true;
+	});
 }
 
 bool FCEClonerMeshBuilder::AppendComponent(UBrushComponent* InComponent, const FTransform& InSourceTransform)
 {
-	return AppendPrimitiveComponent(InComponent, InSourceTransform);
+	return AppendPrimitiveComponent(nullptr, InComponent, InSourceTransform);
 }
 
 bool FCEClonerMeshBuilder::AppendComponent(const USkeletalMeshComponent* InComponent, const FTransform& InSourceTransform)
@@ -306,40 +345,45 @@ bool FCEClonerMeshBuilder::AppendComponent(const USkeletalMeshComponent* InCompo
 		return false;
 	}
 
-	// convert to dynamic mesh
-	FGeometryScriptMeshReadLOD SkeletalMeshLOD;
-	SkeletalMeshLOD.LODType = EGeometryScriptLODType::SourceModel;
-
-	FGeometryScriptCopyMeshFromAssetOptions OutputMeshOptions;
-	OutputMeshOptions.bIgnoreRemoveDegenerates = false;
-	OutputMeshOptions.bRequestTangents = false;
-	OutputMeshOptions.bApplyBuildSettings = false;
+	// Transform the new mesh relative to the component
+	const FTransform RelativeTransform = InComponent->GetComponentTransform().GetRelativeTransform(InSourceTransform);
 
-	EGeometryScriptOutcomePins OutResult;
-	UGeometryScriptLibrary_StaticMeshFunctions::CopyMeshFromSkeletalMesh(SkeletalMesh, OutputDynamicMesh, OutputMeshOptions, SkeletalMeshLOD, OutResult);
+	TArray<TWeakObjectPtr<UMaterialInterface>> Materials;
+	Algo::Transform(InComponent->GetMaterials(), Materials, [](UMaterialInterface* InMaterial)
+	{
+		return InMaterial;
+	});
 
-	if (OutResult != EGeometryScriptOutcomePins::Success)
+	return !!AddMeshInstance(SkeletalMesh->GetUniqueID(), RelativeTransform, Materials, [this, &SkeletalMesh](FDynamicMesh3& InCreateMesh)->bool
 	{
-		return false;
-	}
+		// convert to dynamic mesh
+		FGeometryScriptMeshReadLOD SkeletalMeshLOD;
+		SkeletalMeshLOD.LODType = EGeometryScriptLODType::SourceModel;
 
-	// Transform the new mesh relative to the component
-	const FTransform RelativeTransform = InComponent->GetComponentTransform().GetRelativeTransform(InSourceTransform);
+		FGeometryScriptCopyMeshFromAssetOptions OutputMeshOptions;
+		OutputMeshOptions.bIgnoreRemoveDegenerates = false;
+		OutputMeshOptions.bRequestTangents = false;
+		OutputMeshOptions.bApplyBuildSettings = false;
 
-	// Copy materials
-	OutputMeshMaterials.Append(InComponent->GetMaterials());
+		EGeometryScriptOutcomePins OutResult;
+		UGeometryScriptLibrary_StaticMeshFunctions::CopyMeshFromSkeletalMesh(SkeletalMesh, OutputDynamicMesh, OutputMeshOptions, SkeletalMeshLOD, OutResult);
 
-	OutputDynamicMesh->EditMesh([this, RelativeTransform](FDynamicMesh3& EditMesh)
-	{
-		MeshTransforms::ApplyTransform(EditMesh, RelativeTransform);
-		ConvertedMeshes.Add(MoveTemp(EditMesh));
+		if (OutResult != EGeometryScriptOutcomePins::Success)
+		{
+			return false;
+		}
 
-		// replace by empty mesh for next usage
-		FDynamicMesh3 EmptyMesh;
-		EditMesh = MoveTemp(EmptyMesh);
-	}, EDynamicMeshChangeType::GeneralEdit, EDynamicMeshAttributeChangeFlags::Unknown, /** bDeferChanges */true);
+		OutputDynamicMesh->EditMesh([&InCreateMesh](FDynamicMesh3& EditMesh)
+		{
+			InCreateMesh = MoveTemp(EditMesh);
 
-	return true;
+			// replace by empty mesh for next usage
+			FDynamicMesh3 EmptyMesh;
+			EditMesh = MoveTemp(EmptyMesh);
+		}, EDynamicMeshChangeType::GeneralEdit, EDynamicMeshAttributeChangeFlags::Unknown, /** bDeferChanges */true);
+
+		return true;
+	});
 }
 
 bool FCEClonerMeshBuilder::AppendComponent(UDynamicMeshComponent* InComponent, const FTransform& InSourceTransform)
@@ -351,7 +395,7 @@ bool FCEClonerMeshBuilder::AppendComponent(UDynamicMeshComponent* InComponent, c
 
 	const UDynamicMesh* DynamicMesh = InComponent->GetDynamicMesh();
 
-	if (!IsValid(DynamicMesh))
+	if (!IsValid(DynamicMesh) || DynamicMesh->GetTriangleCount() == 0)
 	{
 		return false;
 	}
@@ -360,17 +404,13 @@ bool FCEClonerMeshBuilder::AppendComponent(UDynamicMeshComponent* InComponent, c
 	const FTransform RelativeTransform = InComponent->GetComponentTransform().GetRelativeTransform(InSourceTransform);
 
 	// Copy all materials
-	OutputMeshMaterials.Append(InComponent->GetMaterials());
-
-	// Create a copy of the mesh
-	DynamicMesh->ProcessMesh([this, RelativeTransform](const FDynamicMesh3& EditMesh)
+	TArray<TWeakObjectPtr<UMaterialInterface>> Materials;
+	Algo::Transform(InComponent->GetMaterials(), Materials, [](UMaterialInterface* InMaterial)
 	{
-		FDynamicMesh3 CopyMesh = EditMesh;
-		MeshTransforms::ApplyTransform(CopyMesh, RelativeTransform);
-		ConvertedMeshes.Add(MoveTemp(CopyMesh));
+		return InMaterial;
 	});
 
-	return true;
+	return AppendMesh(DynamicMesh, Materials, RelativeTransform);
 }
 
 bool FCEClonerMeshBuilder::AppendComponent(UInstancedStaticMeshComponent* InComponent, const FTransform& InSourceTransform)
@@ -380,7 +420,7 @@ bool FCEClonerMeshBuilder::AppendComponent(UInstancedStaticMeshComponent* InComp
 		return false;
 	}
 
-	return AppendPrimitiveComponent(InComponent, InSourceTransform);
+	return AppendPrimitiveComponent(nullptr, InComponent, InSourceTransform);
 }
 
 bool FCEClonerMeshBuilder::AppendComponent(USplineMeshComponent* InComponent, const FTransform& InSourceTransform)
@@ -390,18 +430,170 @@ bool FCEClonerMeshBuilder::AppendComponent(USplineMeshComponent* InComponent, co
 		return false;
 	}
 
-	return AppendPrimitiveComponent(InComponent, InSourceTransform);
+	return AppendPrimitiveComponent(nullptr, InComponent, InSourceTransform);
+}
+
+bool FCEClonerMeshBuilder::AppendComponent(UNiagaraComponent* InComponent, const FTransform& InSourceTransform)
+{
+	if (!IsValid(InComponent))
+	{
+		return false;
+	}
+
+	UNiagaraSystem* System = InComponent->GetAsset();
+
+	if (!IsValid(System))
+	{
+		return false;
+	}
+
+	struct FNiagaraSimCacheEmitterData
+	{
+		TArray<UNiagaraMeshRendererProperties*> MeshRenderers;
+		TArray<FVector> ParticlePositions;
+		TArray<FQuat> ParticleRotations;
+		TArray<FVector> ParticleScales;
+		TArray<int32> ParticleMeshIndexes;
+	};
+
+	TMap<FName, FNiagaraSimCacheEmitterData> EmittersData;
+
+	// Set attributes to capture
+	FNiagaraSimCacheCreateParameters Params;
+	Params.AttributeCaptureMode = ENiagaraSimCacheAttributeCaptureMode::ExplicitAttributes;
+	Params.bAllowDataInterfaceCaching = false;
+
+	for (const FNiagaraEmitterHandle& EmitterHandle : System->GetEmitterHandles())
+	{
+		const FString EmitterName = EmitterHandle.GetUniqueInstanceName();
+
+		Params.ExplicitCaptureAttributes.Add(FName(EmitterName + TEXT(".Particles.Position")));
+		Params.ExplicitCaptureAttributes.Add(FName(EmitterName + TEXT(".Particles.MeshOrientation")));
+		Params.ExplicitCaptureAttributes.Add(FName(EmitterName + TEXT(".Particles.Scale")));
+		Params.ExplicitCaptureAttributes.Add(FName(EmitterName + TEXT(".Particles.MeshIndex")));
+
+		if (const FVersionedNiagaraEmitterData* EmitterData = EmitterHandle.GetEmitterData())
+		{
+			for (UNiagaraRendererProperties* EmitterRenderer : EmitterData->GetRenderers())
+			{
+				if (UNiagaraMeshRendererProperties* MeshRenderer = Cast<UNiagaraMeshRendererProperties>(EmitterRenderer))
+				{
+					FNiagaraSimCacheEmitterData& CacheEmitterData = EmittersData.FindOrAdd(FName(EmitterName));
+					CacheEmitterData.MeshRenderers.Add(MeshRenderer);
+				}
+			}
+		}
+	}
+
+	if (EmittersData.IsEmpty())
+	{
+		return false;
+	}
+
+	UNiagaraSimCache* SimCache = UNiagaraSimCacheFunctionLibrary::CreateNiagaraSimCache(InComponent);
+
+	if (!IsValid(SimCache))
+	{
+		return false;
+	}
+
+	const bool bSuccess = UNiagaraSimCacheFunctionLibrary::CaptureNiagaraSimCacheImmediate(SimCache, Params, InComponent, SimCache, /** AdvanceSim */ false);
+
+	if (!bSuccess)
+	{
+		return false;
+	}
+
+	for (const FNiagaraEmitterHandle& Handle : System->GetEmitterHandles())
+	{
+		constexpr int32 FrameIndex = 0;
+		constexpr bool bLocalToWorld = false;
+		const FName EmitterName(Handle.GetUniqueInstanceName());
+
+		FNiagaraSimCacheEmitterData& CacheEmitterData = EmittersData.FindChecked(EmitterName);
+
+		SimCache->ReadPositionAttribute(CacheEmitterData.ParticlePositions, TEXT("Position"), EmitterName, bLocalToWorld, FrameIndex);
+		SimCache->ReadQuatAttribute(CacheEmitterData.ParticleRotations, TEXT("MeshOrientation"), EmitterName, bLocalToWorld, FrameIndex);
+		SimCache->ReadVectorAttribute(CacheEmitterData.ParticleScales, TEXT("Scale"), EmitterName, FrameIndex);
+		SimCache->ReadIntAttribute(CacheEmitterData.ParticleMeshIndexes, TEXT("MeshIndex"), EmitterName, FrameIndex);
+
+		if (!(CacheEmitterData.ParticlePositions.Num() == CacheEmitterData.ParticleRotations.Num()
+			&& CacheEmitterData.ParticleRotations.Num() == CacheEmitterData.ParticleScales.Num()
+			&& CacheEmitterData.ParticleScales.Num() == CacheEmitterData.ParticleMeshIndexes.Num()))
+		{
+			return false;
+		}
+	}
+
+	const FTransform RelativeTransform = InComponent->GetComponentTransform().GetRelativeTransform(InSourceTransform);
+
+	bool bResult = false;
+
+	for (const TPair<FName, FNiagaraSimCacheEmitterData>& CacheEmitterDataPair : EmittersData)
+	{
+		const FNiagaraSimCacheEmitterData& CacheEmitterData = CacheEmitterDataPair.Value;
+
+		for (int32 Index = 0; Index < CacheEmitterData.ParticlePositions.Num(); Index++)
+		{
+			const FVector& ParticlePosition = CacheEmitterData.ParticlePositions[Index];
+			const FVector& ParticleScale = CacheEmitterData.ParticleScales[Index];
+			const FQuat& ParticleRotation = CacheEmitterData.ParticleRotations[Index];
+			const int32& ParticleMeshIndex = CacheEmitterData.ParticleMeshIndexes[Index];
+
+			FTransform ParticleTransform(ParticleRotation, ParticlePosition, ParticleScale);
+
+			if (ParticleMeshIndex >= 0)
+			{
+				for (UNiagaraMeshRendererProperties* MeshRenderer : CacheEmitterData.MeshRenderers)
+				{
+					if (!IsValid(MeshRenderer))
+					{
+						continue;
+					}
+
+					if (!MeshRenderer->Meshes.IsValidIndex(ParticleMeshIndex))
+					{
+						continue;
+					}
+
+					FNiagaraMeshRendererMeshProperties& Mesh = MeshRenderer->Meshes[ParticleMeshIndex];
+					if (!IsValid(Mesh.Mesh))
+					{
+						continue;
+					}
+
+					FTransform MeshTransform(Mesh.Rotation.Quaternion(), Mesh.PivotOffset, Mesh.Scale);
+					MeshTransform.Accumulate(ParticleTransform);
+					MeshTransform.Accumulate(RelativeTransform);
+
+					TArray<TWeakObjectPtr<UMaterialInterface>> Materials;
+					for (int32 SectionIndex = 0; SectionIndex < Mesh.Mesh->GetNumSections(0); SectionIndex++)
+					{
+						Materials.Add(Mesh.Mesh->GetMaterial(SectionIndex));
+					}
+
+					bResult |= AppendMesh(Mesh.Mesh, Materials, MeshTransform);
+				}
+			}
+		}
+	}
+
+	SimCache->MarkAsGarbage();
+
+	return bResult;
 }
 
-bool FCEClonerMeshBuilder::BuildDynamicMesh(UDynamicMesh* OutMesh, TArray<TWeakObjectPtr<UMaterialInterface>>& OutMaterials)
+bool FCEClonerMeshBuilder::BuildDynamicMesh(UDynamicMesh* OutMesh, TArray<TWeakObjectPtr<UMaterialInterface>>& OutMaterials, const FCEClonerMeshBuilderParams& InParams)
 {
 	if (!IsValid(OutMesh))
 	{
 		return false;
 	}
 
+	OutMaterials.Empty();
+
 	// Lets combine all meshes components from this actor together
-	OutMesh->EditMesh([this](FDynamicMesh3& InMergedMesh)
+	OutMesh->EditMesh([this, &OutMaterials, &InParams](FDynamicMesh3& InMergedMesh)
 	{
 		using namespace UE::Geometry;
 
@@ -414,10 +606,26 @@ bool FCEClonerMeshBuilder::BuildDynamicMesh(UDynamicMesh* OutMesh, TArray<TWeakO
 		AppendOptions.CombineMode = EGeometryScriptCombineAttributesMode::EnableAllMatching;
 		int32 MaterialCount = 0;
 
+		// Mesh index to material forward map to merge same meshes materials into same slot
+		TMap<uint32, TMap<int32, int32>> MeshToMaterialMap;
+
 		// Convert meshes
-		for (int32 MeshIdx = 0; MeshIdx < ConvertedMeshes.Num(); MeshIdx++)
+		for (const FCEClonerMeshInstance& MeshInstance : MeshInstances)
 		{
-			const FDynamicMesh3& ConvertedMesh = ConvertedMeshes[MeshIdx];
+			if (!MeshInstance.Mesh || MeshInstance.Mesh->TriangleCount() == 0)
+			{
+				continue;
+			}
+
+			// Apply transform
+			FDynamicMesh3 ConvertedMesh = *MeshInstance.Mesh;
+			MeshTransforms::ApplyTransform(ConvertedMesh, MeshInstance.MeshData.Transform);
+
+			// Get materials
+			if (!InParams.bMergeMaterials || !MeshToMaterialMap.Contains(MeshInstance.MeshIndex))
+			{
+				OutMaterials.Append(MeshInstance.MeshData.MeshMaterials);
+			}
 
 			// Enable matching attributes & append mesh
 			FMeshIndexMappings TmpMappings;
@@ -429,15 +637,29 @@ bool FCEClonerMeshBuilder::BuildDynamicMesh(UDynamicMesh* OutMesh, TArray<TWeakO
 			{
 				const FDynamicMeshMaterialAttribute* FromMaterialIDAttrib = ConvertedMesh.Attributes()->GetMaterialID();
 				FDynamicMeshMaterialAttribute* ToMaterialIDAttrib = InMergedMesh.Attributes()->GetMaterialID();
-				TMap<int32, int32> MaterialMap;
+
+				TMap<int32, int32>& MaterialMap = MeshToMaterialMap.FindOrAdd(MeshInstance.MeshIndex);
 				for (const TPair<int32, int32>& FromToTId : TmpMappings.GetTriangleMap().GetForwardMap())
 				{
 					const int32 FromMatId = FromMaterialIDAttrib->GetValue(FromToTId.Key);
-					const int32 ToMatId = FromMatId + MaterialCount;
+					int32 ToMatId = FromMatId + MaterialCount;
+
+					// used to merge materials for same mesh index
+					if (const int32* MatId = MaterialMap.Find(FromMatId))
+					{
+						ToMatId = *MatId;
+					}
+
 					MaterialMap.Add(FromMatId, ToMatId);
 					ToMaterialIDAttrib->SetNewValue(FromToTId.Value, ToMatId);
 				}
+
 				MaterialCount += MaterialMap.Num();
+
+				if (!InParams.bMergeMaterials)
+				{
+					MaterialMap.Empty();
+				}
 			}
 		}
 
@@ -449,34 +671,182 @@ bool FCEClonerMeshBuilder::BuildDynamicMesh(UDynamicMesh* OutMesh, TArray<TWeakO
 		}
 	}, EDynamicMeshChangeType::GeneralEdit, EDynamicMeshAttributeChangeFlags::Unknown, /** bDeferChange */true);
 
-	OutMaterials = MoveTemp(OutputMeshMaterials);
+	return true;
+}
+
+bool FCEClonerMeshBuilder::BuildStaticMesh(UStaticMesh* OutMesh, TArray<TWeakObjectPtr<UMaterialInterface>>& OutMaterials, const FCEClonerMeshBuilderParams& InParams)
+{
+	if (!IsValid(OutMesh))
+	{
+		return false;
+	}
+
+	ClearOutputMesh();
+
+	if (!BuildDynamicMesh(OutputDynamicMesh, OutMaterials, InParams))
+	{
+		return false;
+	}
+
+	return DynamicMeshToStaticMesh(OutputDynamicMesh, OutMesh, OutMaterials);
+}
+
+bool FCEClonerMeshBuilder::BuildStaticMesh(int32 InInstanceIndex, UStaticMesh* OutMesh, FCEClonerMeshInstanceData& OutMeshInstance)
+{
+	if (!IsValid(OutMesh) || !MeshInstances.IsValidIndex(InInstanceIndex))
+	{
+		return false;
+	}
+
+	FCEClonerMeshInstance& MeshInstance = MeshInstances[InInstanceIndex];
+
+	if (!MeshInstance.Mesh || MeshInstance.Mesh->TriangleCount() == 0)
+	{
+		return false;
+	}
+
+	OutMeshInstance.MeshMaterials.Empty(MeshInstance.MeshData.MeshMaterials.Num());
+	OutMeshInstance.MeshMaterials.Append(MeshInstance.MeshData.MeshMaterials);
+
+	OutMeshInstance.Transform = MeshInstance.MeshData.Transform;
+
+	OutputDynamicMesh->SetMesh(*MeshInstance.Mesh);
+
+	return DynamicMeshToStaticMesh(OutputDynamicMesh, OutMesh, OutMeshInstance.MeshMaterials);
+}
+
+bool FCEClonerMeshBuilder::BuildDynamicMesh(int32 InInstanceIndex, UDynamicMesh* OutMesh, FCEClonerMeshInstanceData& OutMeshInstance)
+{
+	if (!IsValid(OutMesh) || !MeshInstances.IsValidIndex(InInstanceIndex))
+	{
+		return false;
+	}
+
+	FCEClonerMeshInstance& MeshInstance = MeshInstances[InInstanceIndex];
+
+	if (!MeshInstance.Mesh || MeshInstance.Mesh->TriangleCount() == 0)
+	{
+		return false;
+	}
+
+	OutMeshInstance.MeshMaterials.Empty(MeshInstance.MeshData.MeshMaterials.Num());
+	OutMeshInstance.MeshMaterials.Append(MeshInstance.MeshData.MeshMaterials);
 
-	ConvertedMeshes.Empty();
-	OutputMeshMaterials = {};
+	OutMeshInstance.Transform = MeshInstance.MeshData.Transform;
+
+	OutMesh->SetMesh(*MeshInstance.Mesh);
 
 	return true;
 }
 
-bool FCEClonerMeshBuilder::BuildStaticMesh(UStaticMesh* OutMesh, TArray<TWeakObjectPtr<UMaterialInterface>>& OutMaterials)
+bool FCEClonerMeshBuilder::BuildStaticMesh(uint32 InMeshIndex, UStaticMesh* OutMesh, TArray<FCEClonerMeshInstanceData>& OutMeshInstances)
 {
-	if (!IsValid(OutMesh))
+	const FDynamicMesh3* Mesh = Meshes.Find(InMeshIndex);
+
+	if (!IsValid(OutMesh) || !Mesh)
+	{
+		return false;
+	}
+
+	OutMeshInstances.Empty(MeshInstances.Num());
+	for (const FCEClonerMeshInstance& MeshInstance : MeshInstances)
+	{
+		if (MeshInstance.MeshIndex == InMeshIndex)
+		{
+			OutMeshInstances.Add(MeshInstance.MeshData);
+		}
+	}
+
+	if (OutMeshInstances.IsEmpty())
+	{
+		return false;
+	}
+
+	OutputDynamicMesh->SetMesh(*Mesh);
+
+	return DynamicMeshToStaticMesh(OutputDynamicMesh, OutMesh, OutMeshInstances[0].MeshMaterials);
+}
+
+bool FCEClonerMeshBuilder::BuildDynamicMesh(uint32 InMeshIndex, UDynamicMesh* OutMesh, TArray<FCEClonerMeshInstanceData>& OutMeshInstances)
+{
+	const FDynamicMesh3* Mesh = Meshes.Find(InMeshIndex);
+
+	if (!IsValid(OutMesh) || !Mesh)
 	{
 		return false;
 	}
 
-	UDynamicMesh* OutDynamicMesh = OutputDynamicMesh;
-	if (!BuildDynamicMesh(OutDynamicMesh, OutMaterials))
+	OutMeshInstances.Empty(MeshInstances.Num());
+	for (const FCEClonerMeshInstance& MeshInstance : MeshInstances)
+	{
+		if (MeshInstance.MeshIndex == InMeshIndex)
+		{
+			OutMeshInstances.Add(MeshInstance.MeshData);
+		}
+	}
+
+	if (OutMeshInstances.IsEmpty())
 	{
 		return false;
 	}
 
-	if (OutDynamicMesh->GetTriangleCount() == 0)
+	OutMesh->SetMesh(*Mesh);
+
+	return true;
+}
+
+bool FCEClonerMeshBuilder::AppendPrimitiveComponent(const UObject* InMeshObject, UPrimitiveComponent* InComponent, const FTransform& InSourceTransform)
+{
+	if (!IsValid(InComponent))
 	{
 		return false;
 	}
 
+	// Transform the new mesh relative to the component
+	const FTransform RelativeTransform = InComponent->GetComponentTransform().GetRelativeTransform(InSourceTransform);
+
+	// Copy all materials
+	TArray<TWeakObjectPtr<UMaterialInterface>> Materials;
+	Materials.Reserve(InComponent->GetNumMaterials());
+	for (int32 Index = 0; Index < InComponent->GetNumMaterials(); Index++)
+	{
+		Materials.Add(InComponent->GetMaterial(Index));
+	}
+
+	// Take mesh id or component id to find already converted mesh
+	const uint32 MeshIndex = InMeshObject ? InMeshObject->GetUniqueID() : InComponent->GetUniqueID();
+
+	return !!AddMeshInstance(MeshIndex, RelativeTransform, Materials, [this, &InComponent](FDynamicMesh3& InCreateMesh)->bool
+	{
+		constexpr FGeometryScriptCopyMeshFromComponentOptions Options;
+		FTransform LocalToWorld = FTransform::Identity;
+
+		// convert to dynamic mesh
+		EGeometryScriptOutcomePins OutResult;
+		UGeometryScriptLibrary_SceneUtilityFunctions::CopyMeshFromComponent(InComponent, OutputDynamicMesh, Options, false, LocalToWorld, OutResult);
+
+		if (OutResult != EGeometryScriptOutcomePins::Success)
+		{
+			return false;
+		}
+
+		OutputDynamicMesh->EditMesh([&InCreateMesh](FDynamicMesh3& EditMesh)
+		{
+			InCreateMesh = MoveTemp(EditMesh);
+
+			// replace by empty mesh
+			FDynamicMesh3 EmptyMesh;
+			EditMesh = MoveTemp(EmptyMesh);
+		}, EDynamicMeshChangeType::GeneralEdit, EDynamicMeshAttributeChangeFlags::Unknown, /** bDeferChange */true);
+
+		return true;
+	});
+}
+
+bool FCEClonerMeshBuilder::DynamicMeshToStaticMesh(UDynamicMesh* InMesh, UStaticMesh* OutMesh, const TArray<TWeakObjectPtr<UMaterialInterface>>& InMaterials)
+{
 	TArray<TObjectPtr<UMaterialInterface>> NewMaterials;
-	Algo::Transform(OutMaterials, NewMaterials, [](const TWeakObjectPtr<UMaterialInterface>& InMaterialWeak)
+	Algo::Transform(InMaterials, NewMaterials, [](const TWeakObjectPtr<UMaterialInterface>& InMaterialWeak)
 	{
 		return InMaterialWeak.Get();
 	});
@@ -499,7 +869,7 @@ bool FCEClonerMeshBuilder::BuildStaticMesh(UStaticMesh* OutMesh, TArray<TWeakObj
 
 	// Convert merged mesh to static mesh
 	EGeometryScriptOutcomePins OutResult;
-	UGeometryScriptLibrary_StaticMeshFunctions::CopyMeshToStaticMesh(OutDynamicMesh, OutMesh, AssetOptions, TargetLOD, OutResult);
+	UGeometryScriptLibrary_StaticMeshFunctions::CopyMeshToStaticMesh(InMesh, OutMesh, AssetOptions, TargetLOD, OutResult);
 
 #if WITH_EDITORONLY_DATA
 	// Compute normals and tangents
@@ -508,48 +878,43 @@ bool FCEClonerMeshBuilder::BuildStaticMesh(UStaticMesh* OutMesh, TArray<TWeakObj
 	FStaticMeshOperations::ComputeTangentsAndNormals(*MeshDescription, EComputeNTBsFlags::Normals);
 #endif
 
-	OutDynamicMesh->EditMesh([](FDynamicMesh3& EditMesh){ EditMesh.Clear(); });
-
-	return true;
+	return OutResult == EGeometryScriptOutcomePins::Success;
 }
 
-bool FCEClonerMeshBuilder::AppendPrimitiveComponent(UPrimitiveComponent* InComponent, const FTransform& InSourceTransform)
+void FCEClonerMeshBuilder::ClearOutputMesh() const
 {
-	if (!IsValid(InComponent))
+	if (OutputDynamicMesh)
 	{
-		return false;
+		OutputDynamicMesh->EditMesh([](FDynamicMesh3& EditMesh){ EditMesh.Clear(); });
 	}
+}
 
-	constexpr FGeometryScriptCopyMeshFromComponentOptions Options;
-	FTransform LocalToWorld = FTransform::Identity;
-
-	// convert to dynamic mesh
-	EGeometryScriptOutcomePins OutResult;
-	UGeometryScriptLibrary_SceneUtilityFunctions::CopyMeshFromComponent(InComponent, OutputDynamicMesh, Options, false, LocalToWorld, OutResult);
+FCEClonerMeshBuilder::FCEClonerMeshInstance* FCEClonerMeshBuilder::AddMeshInstance(uint32 InMeshIndex, const FTransform& InTransform, const TArray<TWeakObjectPtr<UMaterialInterface>>& InMaterials, TFunctionRef<bool(UE::Geometry::FDynamicMesh3&)> InCreateMeshFunction)
+{
+	FCEClonerMeshInstance MeshInstance;
+	MeshInstance.MeshIndex = InMeshIndex;
+	MeshInstance.Mesh = Meshes.Find(InMeshIndex);
+	MeshInstance.MeshData.Transform = InTransform;
+	MeshInstance.MeshData.MeshMaterials = InMaterials;
 
-	if (OutResult != EGeometryScriptOutcomePins::Success)
+	if (!MeshInstance.Mesh)
 	{
-		return false;
-	}
+		FDynamicMesh3 Mesh;
 
-	// Transform the new mesh relative to the component
-	const FTransform RelativeTransform = InComponent->GetComponentTransform().GetRelativeTransform(InSourceTransform);
+		ClearOutputMesh();
 
-	// Copy all materials
-	for (int32 Index = 0; Index < InComponent->GetNumMaterials(); Index++)
-	{
-		OutputMeshMaterials.Add(InComponent->GetMaterial(Index));
+		if (InCreateMeshFunction(Mesh) && Mesh.TriangleCount() > 0)
+		{
+			MeshInstance.Mesh = &Meshes.Add(InMeshIndex, MoveTemp(Mesh));
+		}
+
+		ClearOutputMesh();
 	}
 
-	OutputDynamicMesh->EditMesh([this, RelativeTransform](FDynamicMesh3& EditMesh)
+	if (MeshInstance.Mesh)
 	{
-		MeshTransforms::ApplyTransform(EditMesh, RelativeTransform);
-		ConvertedMeshes.Add(MoveTemp(EditMesh));
-
-		// replace by empty mesh
-		FDynamicMesh3 EmptyMesh;
-		EditMesh = MoveTemp(EmptyMesh);
-	}, EDynamicMeshChangeType::GeneralEdit, EDynamicMeshAttributeChangeFlags::Unknown, /** bDeferChange */true);
+		return &MeshInstances.Add_GetRef(MeshInstance);
+	}
 
-	return true;
+	return nullptr;
 }
diff --git a/Engine/Plugins/Experimental/ClonerEffector/Source/ClonerEffector/Private/Effector/CEEffectorActor.cpp b/Engine/Plugins/Experimental/ClonerEffector/Source/ClonerEffector/Private/Effector/CEEffectorActor.cpp
--- a/Engine/Plugins/Experimental/ClonerEffector/Source/ClonerEffector/Private/Effector/CEEffectorActor.cpp
+++ b/Engine/Plugins/Experimental/ClonerEffector/Source/ClonerEffector/Private/Effector/CEEffectorActor.cpp
@@ -98,7 +98,7 @@ void ACEEffectorActor::MigrateDeprecatedProperties()
 		return;
 	}
 
-	UE_LOG(LogCEEffectorActor, Warning, TEXT("%s : Effector migrating from version %i to latest %i"), *GetActorNameOrLabel(), MigrateToVersion, FCEEffectorActorVersion::LatestVersion);
+	UE_LOG(LogCEEffectorActor, Warning, TEXT("%s : Effector migrating from version %i to latest %i, please re-save this asset"), *GetActorNameOrLabel(), MigrateToVersion, FCEEffectorActorVersion::LatestVersion);
 
 	PRAGMA_DISABLE_DEPRECATION_WARNINGS
 
diff --git a/Engine/Plugins/Experimental/ClonerEffector/Source/ClonerEffector/Private/Effector/CEEffectorComponent.cpp b/Engine/Plugins/Experimental/ClonerEffector/Source/ClonerEffector/Private/Effector/CEEffectorComponent.cpp
--- a/Engine/Plugins/Experimental/ClonerEffector/Source/ClonerEffector/Private/Effector/CEEffectorComponent.cpp
+++ b/Engine/Plugins/Experimental/ClonerEffector/Source/ClonerEffector/Private/Effector/CEEffectorComponent.cpp
@@ -436,6 +436,21 @@ void UCEEffectorComponent::OnEffectorDisabled()
 #endif
 }
 
+void UCEEffectorComponent::OnEffectorSetEnabled(const UWorld* InWorld, bool bInEnabled, bool bInTransact)
+{
+	if (GetWorld() == InWorld)
+	{
+#if WITH_EDITOR
+		if (bInTransact)
+		{
+			Modify();
+		}
+#endif
+
+		SetEnabled(bInEnabled);
+	}
+}
+
 void UCEEffectorComponent::OnEffectorOptionsChanged()
 {
 	RegisterToChannel();
diff --git a/Engine/Plugins/Experimental/ClonerEffector/Source/ClonerEffector/Public/CEClonerEffectorShared.h b/Engine/Plugins/Experimental/ClonerEffector/Source/ClonerEffector/Public/CEClonerEffectorShared.h
--- a/Engine/Plugins/Experimental/ClonerEffector/Source/ClonerEffector/Public/CEClonerEffectorShared.h
+++ b/Engine/Plugins/Experimental/ClonerEffector/Source/ClonerEffector/Public/CEClonerEffectorShared.h
@@ -8,9 +8,12 @@
 #include "NiagaraDataInterfaceArrayInt.h"
 #include "CEClonerEffectorShared.generated.h"
 
+class AActor;
+class ADynamicMeshActor;
+class AStaticMeshActor;
+class UCEClonerComponent;
 class UNiagaraDataChannelReader;
 class UNiagaraDataChannelWriter;
-class AActor;
 class UDynamicMesh;
 class UNiagaraSystem;
 class USplineComponent;
@@ -249,6 +252,16 @@ enum class ECEClonerCollisionRadiusMode : uint8
 	ExtentLength
 };
 
+/** Enumerates all conversion possible for cloner simulation */
+enum class ECEClonerMeshConversion : uint8
+{
+	StaticMesh,
+	StaticMeshes,
+	DynamicMesh,
+	DynamicMeshes,
+	InstancedStaticMesh
+};
+
 USTRUCT()
 struct FCEClonerAttachmentItem
 {
@@ -490,8 +503,79 @@ struct CLONEREFFECTOR_API FCEClonerGridConstraintTexture
 	float Threshold = 0.f;
 };
 
+#if WITH_EDITOR
+struct FCEExtensionSection
+{
+	FCEExtensionSection()
+	{}
+
+	explicit FCEExtensionSection(FName InSectionName, int32 InSectionOrder)
+		: SectionName(InSectionName)
+		, SectionOrder(InSectionOrder)
+	{}
+
+	/** Used for editor details UI */
+	FName SectionName = NAME_None;
+
+	/** Used to reorder categories in editor UI */
+	int32 SectionOrder = 0;
+};
+#endif
+
 namespace UE::ClonerEffector
 {
-	void SetBillboardComponentSprite(const AActor* InActor, const FString& InTexturePath);
-	void SetBillboardComponentVisibility(const AActor* InActor, bool bInVisibility);
+#if WITH_EDITOR
+	namespace ClonerSection
+	{
+		const FCEExtensionSection ClonerSection(TEXT("Cloner"), 0);
+		const FCEExtensionSection EffectorSection(TEXT("Effector"), 1);
+		const FCEExtensionSection EmissionSection(TEXT("Emission"), 2);
+		const FCEExtensionSection PhysicsSection(TEXT("Physics"), 3);
+		const FCEExtensionSection RenderingSection(TEXT("Rendering"), 4);
+	}
+
+	namespace EffectorSection
+	{
+		const FCEExtensionSection EffectorSection(TEXT("Effector"), 0);
+		const FCEExtensionSection ModeSection(TEXT("Mode"), 1);
+		const FCEExtensionSection ShapeSection(TEXT("Shape"), 2);
+		const FCEExtensionSection ForcesSection(TEXT("Forces"), 3);
+	}
+#endif
+
+	namespace Conversion
+	{
+		/** Convert a cloner to a single merged static mesh actor */
+		AStaticMeshActor* ConvertClonerToStaticMesh(UCEClonerComponent* InCloner);
+
+		/** Convert a cloner to a single merged dynamic mesh actor */
+		ADynamicMeshActor* ConvertClonerToDynamicMesh(UCEClonerComponent* InCloner);
+
+		/** Convert a cloner to multiple static mesh actors */
+		TArray<AStaticMeshActor*> ConvertClonerToStaticMeshes(UCEClonerComponent* InCloner);
+
+		/** Convert a cloner to multiple dynamic mesh actors */
+		TArray<ADynamicMeshActor*> ConvertClonerToDynamicMeshes(UCEClonerComponent* InCloner);
+
+		/** Convert a cloner to multiple instanced static mesh actors */
+		TArray<AActor*> ConvertClonerToInstancedStaticMeshes(UCEClonerComponent* InCloner);
+
+		UActorComponent* CreateComponent(AActor* InActor, TSubclassOf<USceneComponent> InComponentClass);
+
+#if WITH_EDITOR
+		/** Pick assets location */
+		bool PickAssetPath(const FString& InDefaultPath, FString& OutPickedPath);
+
+		/** Create a specific asset in a package */
+		UObject* CreateAssetPackage(TSubclassOf<UObject> InAssetClass, const FString& InAssetPath);
+
+		template<typename InClass
+			UE_REQUIRES(TIsDerivedFrom<InClass, UObject>::Value)>
+		InClass* CreateAssetPackage(const FString& InAssetPath)
+		{
+			return Cast<InClass>(CreateAssetPackage(InClass::StaticClass(), InAssetPath));
+		}
+#endif
+
+	}
 }
\ No newline at end of file
diff --git a/Engine/Plugins/Experimental/ClonerEffector/Source/ClonerEffector/Public/Cloner/CEClonerActor.h b/Engine/Plugins/Experimental/ClonerEffector/Source/ClonerEffector/Public/Cloner/CEClonerActor.h
--- a/Engine/Plugins/Experimental/ClonerEffector/Source/ClonerEffector/Public/Cloner/CEClonerActor.h
+++ b/Engine/Plugins/Experimental/ClonerEffector/Source/ClonerEffector/Public/Cloner/CEClonerActor.h
@@ -54,6 +54,8 @@ private:
 	void MigrateDeprecatedProperties();
 
 #if WITH_EDITOR
+	bool SpawnDefaultActorAttached(float);
+
 	/** Hide outline selection when this cloner only is selected in viewport */
 	void OnEditorSelectionChanged(UObject* InSelection);
 
diff --git a/Engine/Plugins/Experimental/ClonerEffector/Source/ClonerEffector/Public/Cloner/CEClonerComponent.h b/Engine/Plugins/Experimental/ClonerEffector/Source/ClonerEffector/Public/Cloner/CEClonerComponent.h
--- a/Engine/Plugins/Experimental/ClonerEffector/Source/ClonerEffector/Public/Cloner/CEClonerComponent.h
+++ b/Engine/Plugins/Experimental/ClonerEffector/Source/ClonerEffector/Public/Cloner/CEClonerComponent.h
@@ -5,6 +5,7 @@
 #include "CEClonerEffectorShared.h"
 #include "CEClonerMeshBuilder.h"
 #include "CEPropertyChangeDispatcher.h"
+#include "Containers/Ticker.h"
 #include "Layouts/CEClonerLayoutBase.h"
 #include "NiagaraComponent.h"
 #include "CEClonerComponent.generated.h"
@@ -82,6 +83,15 @@ public:
 		return Seed;
 	}
 
+	UFUNCTION(BlueprintCallable, Category="Cloner")
+	CLONEREFFECTOR_API void SetColor(const FLinearColor& InColor);
+
+	UFUNCTION(BlueprintPure, Category="Cloner")
+	const FLinearColor& GetColor() const
+	{
+		return Color;
+	}
+
 	UFUNCTION(BlueprintCallable, Category="Cloner")
 	CLONEREFFECTOR_API void SetLayoutName(FName InLayoutName);
 
@@ -142,31 +152,39 @@ public:
 	UFUNCTION(CallInEditor, Category="Cloner")
 	void ForceUpdateCloner();
 
-	/** This will create a new default actor attached to this cloner */
-	UFUNCTION(CallInEditor, Category="Cloner")
+	/** This will create a new default actor attached to this cloner if nothing is attached to this cloner */
+	UFUNCTION(CallInEditor, Category="Utilities")
 	void CreateDefaultActorAttached();
+
+	/** Converts the cloner simulation into a single static mesh, this is a heavy operation */
+	UFUNCTION(CallInEditor, Category="Utilities")
+	void ConvertToStaticMesh();
+
+	/** Converts the cloner simulation into a single dynamic mesh, this is a heavy operation */
+	UFUNCTION(CallInEditor, Category="Utilities")
+	void ConvertToDynamicMesh();
+
+	/** Converts the cloner simulation into static meshes, this is a heavy operation */
+	UFUNCTION(CallInEditor, Category="Utilities")
+	void ConvertToStaticMeshes();
+
+	/** Converts the cloner simulation into dynamic meshes, this is a heavy operation */
+	UFUNCTION(CallInEditor, Category="Utilities")
+	void ConvertToDynamicMeshes();
+
+	/** Converts the cloner simulation into instanced static meshes, this is a heavy operation */
+	UFUNCTION(CallInEditor, Category="Utilities")
+	void ConvertToInstancedStaticMeshes();
 #endif
 
 	/** Will force a system update to refresh user parameters */
 	void RequestClonerUpdate(bool bInImmediate = false);
 
-	/**
-	 * Triggers an update of the attachment tree to detect updated items
-	 * If reset is true, clears the attachment tree and rebuilds it otherwise diff update
-	 */
-	void UpdateClonerAttachmentTree(bool bInReset = false);
-
-	/** Called to trigger an update of cloner rendering state tree */
-	void UpdateClonerRenderState();
-
 	UCEClonerLayoutBase* GetClonerActiveLayout() const
 	{
 		return ActiveLayout;
 	}
 
-	/** Forces a refresh of the active system parameters in niagara store */
-	void RefreshUserParameters() const;
-
 	/** Forces a refresh of the meshes used */
 	void RefreshClonerMeshes();
 
@@ -199,11 +217,6 @@ protected:
 	DECLARE_MULTICAST_DELEGATE_OneParam(FOnClonerInitialized, UCEClonerComponent* /** ClonerComponent */)
 	FOnClonerInitialized OnClonerInitializedDelegate;
 
-	//~ Begin UActorComponent
-	/** Tick needed to keep the attachment tree updated */
-	virtual void TickComponent(float InDeltaTime, ELevelTick InTickType, FActorComponentTickFunction* InThisTickFunction) override;
-	//~ End UActorComponent
-
 	//~ Begin UObject
 	virtual void PostInitProperties() override;
 	virtual void PostLoad() override;
@@ -213,6 +226,10 @@ protected:
 #endif
 	//~ End UObject
 
+	//~ Begin UActorComponent
+	virtual void OnComponentCreated() override;
+	//~ End UActorComponent
+
 	void UpdateAttachmentTree();
 	void UpdateActorAttachment(AActor* InActor, AActor* InParent);
 
@@ -283,6 +300,10 @@ protected:
 	UPROPERTY(EditInstanceOnly, BlueprintReadWrite, Setter, Getter, Category="Cloner")
 	int32 Seed = 0;
 
+	/** Cloner color when unaffected by effectors, color will be passed down to the material (ParticleColor) */
+	UPROPERTY(EditInstanceOnly, BlueprintReadWrite, Setter, Getter, Category="Cloner")
+	FLinearColor Color = FLinearColor::White;
+
 	/** Name of the layout to use */
 	UPROPERTY(EditInstanceOnly, BlueprintReadWrite, Setter, Getter, Category="Layout", meta=(GetOptions="GetClonerLayoutNames"))
 	FName LayoutName = NAME_None;
@@ -314,12 +335,28 @@ private:
 
 	/** Initiate and perform operation */
 	void InitializeCloner();
-	void TickCloner(float InDelta);
+	void RegisterTicker();
+	bool TickCloner(float InDelta);
+
+	/**
+	* Triggers an update of the attachment tree to detect updated items
+	* If reset is true, clears the attachment tree and rebuilds it otherwise diff update
+	*/
+	void UpdateClonerAttachmentTree(bool bInReset = false);
+
+	/** Called to trigger an update of cloner rendering state tree */
+	void UpdateClonerRenderState();
+
+	/** Forces a refresh of the active system parameters in niagara store */
+	void RefreshUserParameters() const;
 
 	void OnEnabledChanged();
 	void OnClonerEnabled();
 	void OnClonerDisabled();
+	void OnClonerSetEnabled(const UWorld* InWorld, bool bInEnabled, bool bInTransact);
+
 	void OnSeedChanged();
+	void OnColorChanged();
 	void OnLayoutNameChanged();
 
 #if WITH_EDITOR
@@ -377,6 +414,8 @@ private:
 
 	bool bClonerInitialized = false;
 
+	FTSTicker::FDelegateHandle ClonerTickerHandle;
+
 #if WITH_EDITOR
 	/** Used for PECP */
 	static const TCEPropertyChangeDispatcher<UCEClonerComponent> PropertyChangeDispatcher;
diff --git a/Engine/Plugins/Experimental/ClonerEffector/Source/ClonerEffector/Public/Cloner/CEClonerMeshBuilder.h b/Engine/Plugins/Experimental/ClonerEffector/Source/ClonerEffector/Public/Cloner/CEClonerMeshBuilder.h
--- a/Engine/Plugins/Experimental/ClonerEffector/Source/ClonerEffector/Public/Cloner/CEClonerMeshBuilder.h
+++ b/Engine/Plugins/Experimental/ClonerEffector/Source/ClonerEffector/Public/Cloner/CEClonerMeshBuilder.h
@@ -11,6 +11,7 @@ class UDynamicMesh;
 class UDynamicMeshComponent;
 class UInstancedStaticMeshComponent;
 class UMaterialInterface;
+class UNiagaraComponent;
 class UProceduralMeshComponent;
 class USkeletalMeshComponent;
 class USplineMeshComponent;
@@ -23,8 +24,37 @@ struct FCEClonerMeshBuilder
 {
 	GENERATED_BODY()
 
+	struct FCEClonerMeshInstanceData
+	{
+		/** Transform to apply on the mesh instance */
+		FTransform Transform;
+
+		/** Materials applied on this instance */
+		TArray<TWeakObjectPtr<UMaterialInterface>> MeshMaterials;
+	};
+
+	struct FCEClonerMeshBuilderParams
+	{
+		/** Merge same meshes material slot in the final result */
+		bool bMergeMaterials = false;
+	};
+
+	static const FCEClonerMeshBuilderParams DefaultParams;
+
 	FCEClonerMeshBuilder();
 
+	int32 GetMeshInstanceCount() const
+	{
+		return MeshInstances.Num();
+	}
+
+	int32 GetMeshCount() const
+	{
+		return Meshes.Num();
+	}
+
+	TArray<uint32> GetMeshIndexes() const;
+
 	/** Resets the builder and clears data */
 	void Reset();
 
@@ -52,21 +82,57 @@ struct FCEClonerMeshBuilder
 	/** Appends Spline Mesh component */
 	bool AppendComponent(USplineMeshComponent* InComponent, const FTransform& InTransform = FTransform::Identity);
 
+	/** Appends Niagara component */
+	bool AppendComponent(UNiagaraComponent* InComponent, const FTransform& InTransform = FTransform::Identity);
+
 	/** Appends a dynamic mesh */
 	bool AppendMesh(const UDynamicMesh* InMesh, const TArray<TWeakObjectPtr<UMaterialInterface>>& InMaterials, const FTransform& InTransform = FTransform::Identity);
 
-	/** Builds a dynamic mesh by merging the mesh data imported */
-	bool BuildDynamicMesh(UDynamicMesh* OutMesh, TArray<TWeakObjectPtr<UMaterialInterface>>& OutMaterials);
+	/** Appends a static mesh */
+	bool AppendMesh(UStaticMesh* InMesh, const TArray<TWeakObjectPtr<UMaterialInterface>>& InMaterials, const FTransform& InTransform = FTransform::Identity);
+
+	/** Builds a dynamic mesh by merging all the mesh data imported */
+	bool BuildDynamicMesh(UDynamicMesh* OutMesh, TArray<TWeakObjectPtr<UMaterialInterface>>& OutMaterials, const FCEClonerMeshBuilderParams& InParams = DefaultParams);
+
+	/** Builds a static mesh by merging all the mesh data imported */
+	bool BuildStaticMesh(UStaticMesh* OutMesh, TArray<TWeakObjectPtr<UMaterialInterface>>& OutMaterials, const FCEClonerMeshBuilderParams& InParams = DefaultParams);
 
-	/** Builds a static mesh by merging the mesh data imported */
-	bool BuildStaticMesh(UStaticMesh* OutMesh, TArray<TWeakObjectPtr<UMaterialInterface>>& OutMaterials);
+	/** Builds a static mesh for the specific instance index */
+	bool BuildStaticMesh(int32 InInstanceIndex, UStaticMesh* OutMesh, FCEClonerMeshInstanceData& OutMeshInstance);
+
+	/** Builds a dynamic mesh for the specific instance index */
+	bool BuildDynamicMesh(int32 InInstanceIndex, UDynamicMesh* OutMesh, FCEClonerMeshInstanceData& OutMeshInstance);
+
+	/** Builds a static mesh for the specific mesh index */
+	bool BuildStaticMesh(uint32 InMeshIndex, UStaticMesh* OutMesh, TArray<FCEClonerMeshInstanceData>& OutMeshInstances);
+
+	/** Builds a dynamic mesh for the specific mesh index */
+	bool BuildDynamicMesh(uint32 InMeshIndex, UDynamicMesh* OutMesh, TArray<FCEClonerMeshInstanceData>& OutMeshInstances);
 
 private:
-	bool AppendPrimitiveComponent(UPrimitiveComponent* InComponent, const FTransform& InTransform);
+	static bool DynamicMeshToStaticMesh(UDynamicMesh* InMesh, UStaticMesh* OutMesh, const TArray<TWeakObjectPtr<UMaterialInterface>>& InMaterials);
+
+	struct FCEClonerMeshInstance
+	{
+		/** Index of the mesh to use for this instance */
+		uint32 MeshIndex;
+
+		/** Data linked to this mesh instance */
+		FCEClonerMeshInstanceData MeshData;
+
+		/** Ptr to the actual mesh, equivalent to Meshes[MeshIndex] */
+		UE::Geometry::FDynamicMesh3* Mesh;
+	};
+
+	FCEClonerMeshInstance* AddMeshInstance(uint32 InMeshIndex, const FTransform& InTransform, const TArray<TWeakObjectPtr<UMaterialInterface>>& InMaterials, TFunctionRef<bool(UE::Geometry::FDynamicMesh3&)> InCreateMeshFunction);
+
+	bool AppendPrimitiveComponent(const UObject* InMeshObject, UPrimitiveComponent* InComponent, const FTransform& InTransform);
+
+	void ClearOutputMesh() const;
 
-	TArray<UE::Geometry::FDynamicMesh3> ConvertedMeshes;
+	TMap<uint32, UE::Geometry::FDynamicMesh3> Meshes;
 
-	TArray<TWeakObjectPtr<UMaterialInterface>> OutputMeshMaterials;
+	TArray<FCEClonerMeshInstance> MeshInstances;
 
 	UPROPERTY()
 	TObjectPtr<UDynamicMesh> OutputDynamicMesh;
diff --git a/Engine/Plugins/Experimental/ClonerEffector/Source/ClonerEffector/Public/Effector/CEEffectorComponent.h b/Engine/Plugins/Experimental/ClonerEffector/Source/ClonerEffector/Public/Effector/CEEffectorComponent.h
--- a/Engine/Plugins/Experimental/ClonerEffector/Source/ClonerEffector/Public/Effector/CEEffectorComponent.h
+++ b/Engine/Plugins/Experimental/ClonerEffector/Source/ClonerEffector/Public/Effector/CEEffectorComponent.h
@@ -216,6 +216,8 @@ private:
 	void OnEnabledChanged();
 	void OnEffectorEnabled();
 	void OnEffectorDisabled();
+	void OnEffectorSetEnabled(const UWorld* InWorld, bool bInEnabled, bool bInTransact);
+
 	void OnEffectorOptionsChanged();
 	void OnTypeNameChanged();
 	void OnModeNameChanged();
