MotionDesign : ClonerEffector
- Details panel customization improvements for cloner and effector (hide categories, remove sections, sort ufunctions)
- Improving cloner and effector extension with editor section

diff --git a/Engine/Plugins/Experimental/ClonerEffector/Source/ClonerEffector/Private/Cloner/Extensions/CEClonerEffectorExtension.cpp b/Engine/Plugins/Experimental/ClonerEffector/Source/ClonerEffector/Private/Cloner/Extensions/CEClonerEffectorExtension.cpp
--- a/Engine/Plugins/Experimental/ClonerEffector/Source/ClonerEffector/Private/Cloner/Extensions/CEClonerEffectorExtension.cpp
+++ b/Engine/Plugins/Experimental/ClonerEffector/Source/ClonerEffector/Private/Cloner/Extensions/CEClonerEffectorExtension.cpp
@@ -15,17 +15,6 @@
 
 DEFINE_LOG_CATEGORY_STATIC(LogCEClonerEffectorExtension, Log, All);
 
-void UCEClonerEffectorExtension::SetColor(const FLinearColor& InColor)
-{
-	if (InColor.Equals(Color))
-	{
-		return;
-	}
-
-	Color = InColor;
-	UpdateExtensionParameters(/** UpdateCloner */false);
-}
-
 bool UCEClonerEffectorExtension::LinkEffector(AActor* InEffectorActor)
 {
 	if (!IsValid(InEffectorActor)
@@ -98,13 +87,6 @@ void UCEClonerEffectorExtension::OnExtensionDeactivated()
 	UCEEffectorSubsystem::OnEffectorIdentifierChanged().RemoveAll(this);
 }
 
-void UCEClonerEffectorExtension::OnExtensionParametersChanged(UCEClonerComponent* InComponent)
-{
-	Super::OnExtensionParametersChanged(InComponent);
-
-	InComponent->SetColorParameter(TEXT("EffectorDefaultColor"), Color);
-}
-
 void UCEClonerEffectorExtension::OnEffectorIdentifierChanged(UCEEffectorComponent* InEffector, int32 InOldIdentifier, int32 InNewIdentifier)
 {
 	if (EffectorActorsWeak.Contains(InEffector->GetOwner()))
@@ -257,7 +239,6 @@ void UCEClonerEffectorExtension::CreateLinkedEffector()
 
 const TCEPropertyChangeDispatcher<UCEClonerEffectorExtension> UCEClonerEffectorExtension::PropertyChangeDispatcher =
 {
-	{ GET_MEMBER_NAME_CHECKED(UCEClonerEffectorExtension, Color), &UCEClonerEffectorExtension::OnExtensionPropertyChanged },
 	{ GET_MEMBER_NAME_CHECKED(UCEClonerEffectorExtension, EffectorActorsWeak), &UCEClonerEffectorExtension::OnEffectorActorsChanged }
 };
 
diff --git a/Engine/Plugins/Experimental/ClonerEffector/Source/ClonerEffector/Private/Cloner/Extensions/CEClonerLifetimeExtension.cpp b/Engine/Plugins/Experimental/ClonerEffector/Source/ClonerEffector/Private/Cloner/Extensions/CEClonerLifetimeExtension.cpp
--- a/Engine/Plugins/Experimental/ClonerEffector/Source/ClonerEffector/Private/Cloner/Extensions/CEClonerLifetimeExtension.cpp
+++ b/Engine/Plugins/Experimental/ClonerEffector/Source/ClonerEffector/Private/Cloner/Extensions/CEClonerLifetimeExtension.cpp
@@ -13,8 +13,7 @@ UCEClonerLifetimeExtension::UCEClonerLifetimeExtension()
 		TEXT("Lifetime")
 		, 0
 #if WITH_EDITOR
-		, TEXT("Emitter")
-		, 4
+		, UE::ClonerEffector::ClonerSection::EmissionSection
 #endif
 		)
 {
diff --git a/Engine/Plugins/Experimental/ClonerEffector/Source/ClonerEffector/Private/Cloner/Extensions/CEClonerMeshRendererExtension.cpp b/Engine/Plugins/Experimental/ClonerEffector/Source/ClonerEffector/Private/Cloner/Extensions/CEClonerMeshRendererExtension.cpp
--- a/Engine/Plugins/Experimental/ClonerEffector/Source/ClonerEffector/Private/Cloner/Extensions/CEClonerMeshRendererExtension.cpp
+++ b/Engine/Plugins/Experimental/ClonerEffector/Source/ClonerEffector/Private/Cloner/Extensions/CEClonerMeshRendererExtension.cpp
@@ -16,8 +16,7 @@ UCEClonerMeshRendererExtension::UCEClonerMeshRendererExtension()
 		TEXT("MeshRenderer")
 		, 1
 #if WITH_EDITOR
-		, TEXT("MeshRenderer")
-		, 2
+		, UE::ClonerEffector::ClonerSection::RenderingSection
 #endif
 		)
 {
diff --git a/Engine/Plugins/Experimental/ClonerEffector/Source/ClonerEffector/Public/Cloner/Extensions/CEClonerCollisionExtension.h b/Engine/Plugins/Experimental/ClonerEffector/Source/ClonerEffector/Public/Cloner/Extensions/CEClonerCollisionExtension.h
--- a/Engine/Plugins/Experimental/ClonerEffector/Source/ClonerEffector/Public/Cloner/Extensions/CEClonerCollisionExtension.h
+++ b/Engine/Plugins/Experimental/ClonerEffector/Source/ClonerEffector/Public/Cloner/Extensions/CEClonerCollisionExtension.h
@@ -19,8 +19,7 @@ public:
 			TEXT("Collisions")
 			, 0
 #if WITH_EDITOR
-			, TEXT("Collisions")
-			, 3
+			, UE::ClonerEffector::ClonerSection::PhysicsSection
 #endif
 		)
 	{}
diff --git a/Engine/Plugins/Experimental/ClonerEffector/Source/ClonerEffector/Public/Cloner/Extensions/CEClonerEffectorExtension.h b/Engine/Plugins/Experimental/ClonerEffector/Source/ClonerEffector/Public/Cloner/Extensions/CEClonerEffectorExtension.h
--- a/Engine/Plugins/Experimental/ClonerEffector/Source/ClonerEffector/Public/Cloner/Extensions/CEClonerEffectorExtension.h
+++ b/Engine/Plugins/Experimental/ClonerEffector/Source/ClonerEffector/Public/Cloner/Extensions/CEClonerEffectorExtension.h
@@ -23,21 +23,11 @@ public:
 			TEXT("Effector")
 			, 0
 #if WITH_EDITOR
-			, TEXT("Effector")
-			, 1
+			, UE::ClonerEffector::ClonerSection::EffectorSection
 #endif
 		)
 	{}
 
-	UFUNCTION(BlueprintCallable, Category="Cloner")
-	CLONEREFFECTOR_API void SetColor(const FLinearColor& InColor);
-
-	UFUNCTION(BlueprintPure, Category="Cloner")
-	const FLinearColor& GetColor() const
-	{
-		return Color;
-	}
-
 	/** Gets the number of effectors applied on this cloner */
 	UFUNCTION(BlueprintPure, Category="Cloner")
 	CLONEREFFECTOR_API int32 GetEffectorCount() const;
@@ -71,17 +61,12 @@ protected:
 	//~ Begin UCEClonerExtensionBase
 	virtual void OnExtensionActivated() override;
 	virtual void OnExtensionDeactivated() override;
-	virtual void OnExtensionParametersChanged(UCEClonerComponent* InComponent) override;
 	//~ End UCEClonerExtensionBase
 
 	void OnEffectorIdentifierChanged(UCEEffectorComponent* InEffector, int32 InOldIdentifier, int32 InNewIdentifier);
 	void OnEffectorsChanged();
 	void OnEffectorActorsChanged();
 
-	/** Cloner color when unaffected by effectors, color will be passed down to the material (ParticleColor) */
-	UPROPERTY(EditInstanceOnly, BlueprintReadWrite, Setter, Getter, Category="Effector")
-	FLinearColor Color = FLinearColor::White;
-
 	/** Effectors actors linked to this cloner */
 	UPROPERTY(EditInstanceOnly, Category="Effector", meta=(DisplayName="Effectors"))
 	TArray<TWeakObjectPtr<AActor>> EffectorActorsWeak;
diff --git a/Engine/Plugins/Experimental/ClonerEffector/Source/ClonerEffector/Public/Cloner/Extensions/CEClonerEmitterSpawnExtension.h b/Engine/Plugins/Experimental/ClonerEffector/Source/ClonerEffector/Public/Cloner/Extensions/CEClonerEmitterSpawnExtension.h
--- a/Engine/Plugins/Experimental/ClonerEffector/Source/ClonerEffector/Public/Cloner/Extensions/CEClonerEmitterSpawnExtension.h
+++ b/Engine/Plugins/Experimental/ClonerEffector/Source/ClonerEffector/Public/Cloner/Extensions/CEClonerEmitterSpawnExtension.h
@@ -21,8 +21,7 @@ public:
 			TEXT("Spawn")
 			, 0
 #if WITH_EDITOR
-			, TEXT("Emitter")
-			, 4
+			, UE::ClonerEffector::ClonerSection::EmissionSection
 #endif
 		)
 	{}
diff --git a/Engine/Plugins/Experimental/ClonerEffector/Source/ClonerEffector/Public/Cloner/Extensions/CEClonerExtensionBase.h b/Engine/Plugins/Experimental/ClonerEffector/Source/ClonerEffector/Public/Cloner/Extensions/CEClonerExtensionBase.h
--- a/Engine/Plugins/Experimental/ClonerEffector/Source/ClonerEffector/Public/Cloner/Extensions/CEClonerExtensionBase.h
+++ b/Engine/Plugins/Experimental/ClonerEffector/Source/ClonerEffector/Public/Cloner/Extensions/CEClonerExtensionBase.h
@@ -2,6 +2,7 @@
 
 #pragma once
 
+#include "CEClonerEffectorShared.h"
 #include "UObject/Object.h"
 #include "CEClonerExtensionBase.generated.h"
 
@@ -20,8 +21,7 @@ public:
 			NAME_None
 			, 0
 #if WITH_EDITOR
-			, NAME_None
-			, 0
+			, FCEExtensionSection(NAME_None, INDEX_NONE)
 #endif
 		)
 	{}
@@ -30,15 +30,13 @@ public:
 		FName InExtensionName
 		, int32 InExtensionPriority
 #if WITH_EDITOR
-		, FName InExtensionCategory
-		, int32 InExtensionCategoryOrder
+		, const FCEExtensionSection& InExtensionSection
 #endif
 		)
 		: ExtensionName(InExtensionName)
 		, ExtensionPriority(InExtensionPriority)
 #if WITH_EDITOR
-		, ExtensionCategory(InExtensionCategory)
-		, ExtensionCategoryOrder(InExtensionCategoryOrder)
+		, ExtensionSection(InExtensionSection)
 #endif
 	{}
 
@@ -53,14 +51,9 @@ public:
 	}
 
 #if WITH_EDITOR
-	FName GetExtensionCategory() const
-	{
-		return ExtensionCategory;
-	}
-
-	int32 GetExtensionCategoryOrder() const
+	const FCEExtensionSection& GetExtensionSection() const
 	{
-		return ExtensionCategoryOrder;
+		return ExtensionSection;
 	}
 #endif
 
@@ -121,12 +114,8 @@ private:
 
 	bool bExtensionActive = false;
 
-#if WITH_EDITORONLY_DATA
+#if WITH_EDITOR
 	/** Used for editor UI */
-	UPROPERTY(Transient)
-	FName ExtensionCategory = NAME_None;
-
-	/** Used to reorder categories in editor UI */
-	int32 ExtensionCategoryOrder = 0;
+	FCEExtensionSection ExtensionSection;
 #endif
 };
\ No newline at end of file
diff --git a/Engine/Plugins/Experimental/ClonerEffector/Source/ClonerEffector/Public/Cloner/Extensions/CEClonerProgressExtension.h b/Engine/Plugins/Experimental/ClonerEffector/Source/ClonerEffector/Public/Cloner/Extensions/CEClonerProgressExtension.h
--- a/Engine/Plugins/Experimental/ClonerEffector/Source/ClonerEffector/Public/Cloner/Extensions/CEClonerProgressExtension.h
+++ b/Engine/Plugins/Experimental/ClonerEffector/Source/ClonerEffector/Public/Cloner/Extensions/CEClonerProgressExtension.h
@@ -18,8 +18,7 @@ public:
 			TEXT("Progress")
 			, 0
 #if WITH_EDITOR
-			, TEXT("Emitter")
-			, 4
+			, UE::ClonerEffector::ClonerSection::EmissionSection
 #endif
 		)
 	{}
diff --git a/Engine/Plugins/Experimental/ClonerEffector/Source/ClonerEffector/Public/Cloner/Extensions/CEClonerRangeExtension.h b/Engine/Plugins/Experimental/ClonerEffector/Source/ClonerEffector/Public/Cloner/Extensions/CEClonerRangeExtension.h
--- a/Engine/Plugins/Experimental/ClonerEffector/Source/ClonerEffector/Public/Cloner/Extensions/CEClonerRangeExtension.h
+++ b/Engine/Plugins/Experimental/ClonerEffector/Source/ClonerEffector/Public/Cloner/Extensions/CEClonerRangeExtension.h
@@ -18,8 +18,7 @@ public:
 			TEXT("Range")
 			, 0
 #if WITH_EDITOR
-			, TEXT("Range")
-			, 5
+			, UE::ClonerEffector::ClonerSection::ClonerSection
 #endif
 		)
 	{}
@@ -142,7 +141,7 @@ protected:
 	FRotator RangeRotationMax = FRotator::ZeroRotator;
 
 	UPROPERTY(EditInstanceOnly, BlueprintReadWrite, Setter="SetRangeScaleUniform", Getter="GetRangeScaleUniform", DisplayName="ScaleUniformEnabled", Category="Range", meta=(EditCondition="bRangeEnabled", EditConditionHides))
-	bool bRangeScaleUniform = false;
+	bool bRangeScaleUniform = true;
 
 	UPROPERTY(EditInstanceOnly, BlueprintReadWrite, Setter="SetRangeScaleMin", Getter="GetRangeScaleMin", DisplayName="ScaleMin", Category="Range", meta=(AllowPreserveRatio, Delta="0.0001", ClampMin="0", EditCondition="bRangeEnabled && !bRangeScaleUniform", EditConditionHides))
 	FVector RangeScaleMin = FVector::OneVector;
diff --git a/Engine/Plugins/Experimental/ClonerEffector/Source/ClonerEffector/Public/Cloner/Extensions/CEClonerStepExtension.h b/Engine/Plugins/Experimental/ClonerEffector/Source/ClonerEffector/Public/Cloner/Extensions/CEClonerStepExtension.h
--- a/Engine/Plugins/Experimental/ClonerEffector/Source/ClonerEffector/Public/Cloner/Extensions/CEClonerStepExtension.h
+++ b/Engine/Plugins/Experimental/ClonerEffector/Source/ClonerEffector/Public/Cloner/Extensions/CEClonerStepExtension.h
@@ -18,8 +18,7 @@ public:
 			TEXT("Step")
 			, 0
 #if WITH_EDITOR
-			, TEXT("Step")
-			, 5
+			, UE::ClonerEffector::ClonerSection::ClonerSection
 #endif
 		)
 	{}
diff --git a/Engine/Plugins/Experimental/ClonerEffector/Source/ClonerEffector/Public/Effector/CEEffectorExtensionBase.h b/Engine/Plugins/Experimental/ClonerEffector/Source/ClonerEffector/Public/Effector/CEEffectorExtensionBase.h
--- a/Engine/Plugins/Experimental/ClonerEffector/Source/ClonerEffector/Public/Effector/CEEffectorExtensionBase.h
+++ b/Engine/Plugins/Experimental/ClonerEffector/Source/ClonerEffector/Public/Effector/CEEffectorExtensionBase.h
@@ -2,6 +2,7 @@
 
 #pragma once
 
+#include "CEClonerEffectorShared.h"
 #include "UObject/Object.h"
 #include "CEEffectorExtensionBase.generated.h"
 
@@ -15,13 +16,23 @@ class UCEEffectorExtensionBase : public UObject
 
 public:
 	UCEEffectorExtensionBase()
-		: UCEEffectorExtensionBase(NAME_None, NAME_None)
+		: UCEEffectorExtensionBase(
+			NAME_None
+#if WITH_EDITOR
+			, FCEExtensionSection(NAME_None, INDEX_NONE)
+#endif
+		)
 	{}
 
-	UCEEffectorExtensionBase(FName InExtensionName, FName InExtensionCategory)
+	UCEEffectorExtensionBase(
+		FName InExtensionName
+#if WITH_EDITOR
+		, const FCEExtensionSection& InExtensionSection
+#endif
+		)
 		: ExtensionName(InExtensionName)
-#if WITH_EDITORONLY_DATA
-		, ExtensionCategory(InExtensionCategory)
+#if WITH_EDITOR
+		, ExtensionSection(InExtensionSection)
 #endif
 	{}
 
@@ -31,9 +42,9 @@ public:
 	}
 
 #if WITH_EDITOR
-	FName GetExtensionCategory() const
+	const FCEExtensionSection& GetExtensionSection() const
 	{
-		return ExtensionCategory;
+		return ExtensionSection;
 	}
 #endif
 
@@ -79,6 +90,6 @@ private:
 	bool bExtensionActive = false;
 
 #if WITH_EDITOR
-	FName ExtensionCategory = NAME_None;
+	FCEExtensionSection ExtensionSection;
 #endif
 };
\ No newline at end of file
diff --git a/Engine/Plugins/Experimental/ClonerEffector/Source/ClonerEffector/Public/Effector/Extensions/CEEffectorForceExtension.h b/Engine/Plugins/Experimental/ClonerEffector/Source/ClonerEffector/Public/Effector/Extensions/CEEffectorForceExtension.h
--- a/Engine/Plugins/Experimental/ClonerEffector/Source/ClonerEffector/Public/Effector/Extensions/CEEffectorForceExtension.h
+++ b/Engine/Plugins/Experimental/ClonerEffector/Source/ClonerEffector/Public/Effector/Extensions/CEEffectorForceExtension.h
@@ -16,7 +16,12 @@ class UCEEffectorForceExtension : public UCEEffectorExtensionBase
 
 public:
 	UCEEffectorForceExtension()
-		: UCEEffectorExtensionBase(TEXT("Forces"), TEXT("Forces"))
+		: UCEEffectorExtensionBase(
+			TEXT("Forces")
+#if WITH_EDITOR
+			, UE::ClonerEffector::EffectorSection::ForcesSection
+#endif
+		)
 	{}
 
 	UFUNCTION(BlueprintCallable, Category="Effector")
diff --git a/Engine/Plugins/Experimental/ClonerEffector/Source/ClonerEffector/Public/Effector/Modes/CEEffectorModeBase.h b/Engine/Plugins/Experimental/ClonerEffector/Source/ClonerEffector/Public/Effector/Modes/CEEffectorModeBase.h
--- a/Engine/Plugins/Experimental/ClonerEffector/Source/ClonerEffector/Public/Effector/Modes/CEEffectorModeBase.h
+++ b/Engine/Plugins/Experimental/ClonerEffector/Source/ClonerEffector/Public/Effector/Modes/CEEffectorModeBase.h
@@ -19,7 +19,12 @@ public:
 	{}
 
 	UCEEffectorModeBase(FName InModeName, int32 InModeIdentifier)
-		: UCEEffectorExtensionBase(InModeName, TEXT("Mode"))
+		: UCEEffectorExtensionBase(
+			InModeName
+#if WITH_EDITOR
+			, UE::ClonerEffector::EffectorSection::ModeSection
+#endif
+		)
 		, ModeIdentifier(InModeIdentifier)
 	{}
 
diff --git a/Engine/Plugins/Experimental/ClonerEffector/Source/ClonerEffector/Public/Effector/Types/CEEffectorTypeBase.h b/Engine/Plugins/Experimental/ClonerEffector/Source/ClonerEffector/Public/Effector/Types/CEEffectorTypeBase.h
--- a/Engine/Plugins/Experimental/ClonerEffector/Source/ClonerEffector/Public/Effector/Types/CEEffectorTypeBase.h
+++ b/Engine/Plugins/Experimental/ClonerEffector/Source/ClonerEffector/Public/Effector/Types/CEEffectorTypeBase.h
@@ -29,7 +29,12 @@ public:
 	{}
 
 	UCEEffectorTypeBase(FName InTypeName, int32 InTypeIdentifier)
-		: UCEEffectorExtensionBase(InTypeName, TEXT("Shape"))
+		: UCEEffectorExtensionBase(
+			InTypeName
+#if WITH_EDITOR
+			, UE::ClonerEffector::EffectorSection::ShapeSection
+#endif
+		)
 		, TypeIdentifier(InTypeIdentifier)
 	{}
 
diff --git a/Engine/Plugins/Experimental/ClonerEffector/Source/ClonerEffectorEditor/Private/CEEditorModule.cpp b/Engine/Plugins/Experimental/ClonerEffector/Source/ClonerEffectorEditor/Private/CEEditorModule.cpp
--- a/Engine/Plugins/Experimental/ClonerEffector/Source/ClonerEffectorEditor/Private/CEEditorModule.cpp
+++ b/Engine/Plugins/Experimental/ClonerEffector/Source/ClonerEffectorEditor/Private/CEEditorModule.cpp
@@ -2,11 +2,15 @@
 
 #include "CEEditorModule.h"
 
-#include "Cloner/CEEditorClonerComponentDetailCustomization.h"
 #include "Cloner/CEClonerComponent.h"
-#include "Effector/CEEditorEffectorComponentDetailCustomization.h"
-#include "Effector/CEEditorEffectorTypeDetailCustomization.h"
+#include "Cloner/CEClonerActor.h"
+#include "Cloner/Customizations/CEEditorClonerComponentDetailCustomization.h"
+#include "Cloner/Customizations/CEEditorClonerActorDetailCustomization.h"
+#include "Effector/Customizations/CEEditorEffectorComponentDetailCustomization.h"
+#include "Effector/Customizations/CEEditorEffectorTypeDetailCustomization.h"
+#include "Effector/CEEffectorActor.h"
 #include "Effector/CEEffectorComponent.h"
+#include "Effector/Customizations/CEEditorEffectorActorDetailCustomization.h"
 #include "Effector/Types/CEEffectorBoundType.h"
 #include "Modules/ModuleManager.h"
 #include "PropertyEditorModule.h"
@@ -19,8 +23,12 @@ void FCEEditorModule::StartupModule()
 	// Load styles
 	FCEEditorStyle::Get();
 
-	// Cloner/effector customization
+	// Cloner customization
+	PropertyModule.RegisterCustomClassLayout(ACEClonerActor::StaticClass()->GetFName(), FOnGetDetailCustomizationInstance::CreateStatic(&FCEEditorClonerActorDetailCustomization::MakeInstance));
 	PropertyModule.RegisterCustomClassLayout(UCEClonerComponent::StaticClass()->GetFName(), FOnGetDetailCustomizationInstance::CreateStatic(&FCEEditorClonerComponentDetailCustomization::MakeInstance));
+
+	// Effector customization
+	PropertyModule.RegisterCustomClassLayout(ACEEffectorActor::StaticClass()->GetFName(), FOnGetDetailCustomizationInstance::CreateStatic(&FCEEditorEffectorActorDetailCustomization::MakeInstance));
 	PropertyModule.RegisterCustomClassLayout(UCEEffectorComponent::StaticClass()->GetFName(), FOnGetDetailCustomizationInstance::CreateStatic(&FCEEditorEffectorComponentDetailCustomization::MakeInstance));
 	PropertyModule.RegisterCustomClassLayout(UCEEffectorBoundType::StaticClass()->GetFName(), FOnGetDetailCustomizationInstance::CreateStatic(&FCEEditorEffectorTypeDetailCustomization::MakeInstance));
 }
@@ -31,8 +39,12 @@ void FCEEditorModule::ShutdownModule()
 	{
 		FPropertyEditorModule& PropertyModule = FModuleManager::GetModuleChecked<FPropertyEditorModule>(PropertyEditorName);
 
-		// Cloner/effector
+		// Cloner customization
+		PropertyModule.UnregisterCustomClassLayout(ACEClonerActor::StaticClass()->GetFName());
 		PropertyModule.UnregisterCustomClassLayout(UCEClonerComponent::StaticClass()->GetFName());
+
+		// Effector customization
+		PropertyModule.UnregisterCustomClassLayout(ACEEffectorActor::StaticClass()->GetFName());
 		PropertyModule.UnregisterCustomClassLayout(UCEEffectorComponent::StaticClass()->GetFName());
 		PropertyModule.UnregisterCustomClassLayout(UCEEffectorBoundType::StaticClass()->GetFName());
 	}
diff --git a/Engine/Plugins/Experimental/ClonerEffector/Source/ClonerEffectorEditor/Private/Cloner/CEEditorClonerComponentDetailCustomization.cpp b/Engine/Plugins/Experimental/ClonerEffector/Source/ClonerEffectorEditor/Private/Cloner/CEEditorClonerComponentDetailCustomization.cpp
deleted file mode 100644
diff --git a/Engine/Plugins/Experimental/ClonerEffector/Source/ClonerEffectorEditor/Private/Cloner/CEEditorClonerComponentDetailCustomization.h b/Engine/Plugins/Experimental/ClonerEffector/Source/ClonerEffectorEditor/Private/Cloner/CEEditorClonerComponentDetailCustomization.h
deleted file mode 100644
diff --git a/Engine/Plugins/Experimental/ClonerEffector/Source/ClonerEffectorEditor/Private/Cloner/Customizations/CEEditorClonerActorDetailCustomization.cpp b/Engine/Plugins/Experimental/ClonerEffector/Source/ClonerEffectorEditor/Private/Cloner/Customizations/CEEditorClonerActorDetailCustomization.cpp
new file mode 100644
--- /dev/null
+++ b/Engine/Plugins/Experimental/ClonerEffector/Source/ClonerEffectorEditor/Private/Cloner/Customizations/CEEditorClonerActorDetailCustomization.cpp
@@ -0,0 +1,34 @@
+// Copyright Epic Games, Inc. All Rights Reserved.
+
+#include "Cloner/Customizations/CEEditorClonerActorDetailCustomization.h"
+
+#include "Cloner/CEClonerActor.h"
+#include "Modules/ModuleManager.h"
+#include "PropertyEditorModule.h"
+
+void FCEEditorClonerActorDetailCustomization::CustomizeDetails(IDetailLayoutBuilder& InDetailBuilder)
+{
+	// Add customizations for actor here
+}
+
+void FCEEditorClonerActorDetailCustomization::RemoveEmptySections()
+{
+	FPropertyEditorModule& PropertyModule = FModuleManager::GetModuleChecked<FPropertyEditorModule>("PropertyEditor");
+	const FName ComponentClassName = ACEClonerActor::StaticClass()->GetFName();
+
+	// Remove Streaming section by removing its categories
+	TSharedRef<FPropertySection> StreamingSection = PropertyModule.FindOrCreateSection(ComponentClassName, "Streaming", FText::GetEmpty());
+	StreamingSection->RemoveCategory("WorldPartition");
+	StreamingSection->RemoveCategory("DataLayers");
+	StreamingSection->RemoveCategory("HLOD");
+
+	// Remove Effects section by removing its categories
+	TSharedRef<FPropertySection> EffectsSection = PropertyModule.FindOrCreateSection(ComponentClassName, "Effects", FText::GetEmpty());
+	EffectsSection->RemoveCategory("Niagara");
+	EffectsSection->RemoveCategory("NiagaraComponent_Parameters");
+	EffectsSection->RemoveCategory("NiagaraComponent_Utilities");
+	EffectsSection->RemoveCategory("Activation");
+	EffectsSection->RemoveCategory("Parameters");
+	EffectsSection->RemoveCategory("Randomness");
+	EffectsSection->RemoveCategory("Warmup");
+}
diff --git a/Engine/Plugins/Experimental/ClonerEffector/Source/ClonerEffectorEditor/Private/Cloner/Customizations/CEEditorClonerActorDetailCustomization.h b/Engine/Plugins/Experimental/ClonerEffector/Source/ClonerEffectorEditor/Private/Cloner/Customizations/CEEditorClonerActorDetailCustomization.h
new file mode 100644
--- /dev/null
+++ b/Engine/Plugins/Experimental/ClonerEffector/Source/ClonerEffectorEditor/Private/Cloner/Customizations/CEEditorClonerActorDetailCustomization.h
@@ -0,0 +1,28 @@
+// Copyright Epic Games, Inc. All Rights Reserved.
+
+#pragma once
+
+#include "IDetailCustomization.h"
+#include "Templates/SharedPointer.h"
+
+/** Used to customize cloner actor properties in details panel */
+class FCEEditorClonerActorDetailCustomization : public IDetailCustomization
+{
+public:
+	static TSharedRef<IDetailCustomization> MakeInstance()
+	{
+		return MakeShared<FCEEditorClonerActorDetailCustomization>();
+	}
+
+	explicit FCEEditorClonerActorDetailCustomization()
+	{
+		RemoveEmptySections();
+	}
+
+	//~ Begin IDetailCustomization
+	virtual void CustomizeDetails(IDetailLayoutBuilder& InDetailBuilder) override;
+	//~ End IDetailCustomization
+
+protected:
+	static void RemoveEmptySections();
+};
diff --git a/Engine/Plugins/Experimental/ClonerEffector/Source/ClonerEffectorEditor/Private/Cloner/Customizations/CEEditorClonerComponentDetailCustomization.cpp b/Engine/Plugins/Experimental/ClonerEffector/Source/ClonerEffectorEditor/Private/Cloner/Customizations/CEEditorClonerComponentDetailCustomization.cpp
new file mode 100644
--- /dev/null
+++ b/Engine/Plugins/Experimental/ClonerEffector/Source/ClonerEffectorEditor/Private/Cloner/Customizations/CEEditorClonerComponentDetailCustomization.cpp
@@ -0,0 +1,279 @@
+// Copyright Epic Games, Inc. All Rights Reserved.
+
+#include "Cloner/Customizations/CEEditorClonerComponentDetailCustomization.h"
+
+#include "Cloner/CEClonerComponent.h"
+#include "Cloner/Extensions/CEClonerExtensionBase.h"
+#include "Cloner/Extensions/CEClonerLifetimeExtension.h"
+#include "Cloner/Layouts/CEClonerLayoutBase.h"
+#include "DetailLayoutBuilder.h"
+#include "DetailWidgetRow.h"
+#include "Input/Reply.h"
+#include "Modules/ModuleManager.h"
+#include "NiagaraDataInterfaceCurve.h"
+#include "PropertyEditorModule.h"
+#include "UObject/Class.h"
+#include "UObject/Object.h"
+#include "Widgets/Input/SButton.h"
+
+#define LOCTEXT_NAMESPACE "CEEditorClonerComponentDetailCustomization"
+
+void FCEEditorClonerComponentDetailCustomization::CustomizeDetails(IDetailLayoutBuilder& InDetailBuilder)
+{
+	FPropertyEditorModule& PropertyModule = FModuleManager::GetModuleChecked<FPropertyEditorModule>(TEXT("PropertyEditor"));
+	const FName ComponentClassName = UCEClonerComponent::StaticClass()->GetFName();
+
+	// Remove exposed user parameters
+	InDetailBuilder.HideCategory(TEXT("NiagaraComponent_Parameters"));
+
+	// Remove niagara utilities
+	InDetailBuilder.HideCategory(TEXT("NiagaraComponent_Utilities"));
+
+	// Remove extension array property
+	TSharedRef<IPropertyHandle> ActiveExtensionsProperty = InDetailBuilder.GetProperty(GET_MEMBER_NAME_CHECKED(UCEClonerComponent, ActiveExtensions), UCEClonerComponent::StaticClass());
+
+	if (!ActiveExtensionsProperty->IsValidHandle())
+	{
+		return;
+	}
+
+	InDetailBuilder.HideProperty(ActiveExtensionsProperty);
+
+	uint32 NumElements = 0;
+	TSharedPtr<IPropertyHandleArray> ActiveExtensionArrayProperty = ActiveExtensionsProperty->AsArray();
+	ActiveExtensionArrayProperty->GetNumElements(NumElements);
+
+	// Everything needs to be below Cloner category
+	int32 StartOrder = InDetailBuilder.EditCategory(TEXT("Cloner")).GetSortOrder() + 1;
+
+	TSharedRef<IPropertyHandle> ActiveLayoutProperty = InDetailBuilder.GetProperty(GET_MEMBER_NAME_CHECKED(UCEClonerComponent, ActiveLayout), UCEClonerComponent::StaticClass());
+
+	if (!ActiveLayoutProperty->IsValidHandle())
+	{
+		return;
+	}
+
+	const FName LayoutCategoryName = ActiveLayoutProperty->GetDefaultCategoryName();
+	IDetailCategoryBuilder& LayoutCategory = InDetailBuilder.EditCategory(LayoutCategoryName);
+	LayoutCategory.SetSortOrder(StartOrder++);
+
+	const TSharedRef<FPropertySection> ClonerSection = PropertyModule.FindOrCreateSection(ComponentClassName, UE::ClonerEffector::ClonerSection::ClonerSection.SectionName, FText::FromName(UE::ClonerEffector::ClonerSection::ClonerSection.SectionName));
+	ClonerSection->AddCategory(LayoutCategoryName);
+
+	const TArray<TWeakObjectPtr<UCEClonerComponent>> ClonerComponentsWeak = InDetailBuilder.GetObjectsOfTypeBeingCustomized<UCEClonerComponent>();
+
+	for (uint32 ElementIndex = 0; ElementIndex < NumElements; ElementIndex++)
+	{
+		TSharedPtr<IPropertyHandle> ActiveExtensionProperty = ActiveExtensionArrayProperty->GetElement(ElementIndex);
+
+		UObject* ObjectValue;
+		ActiveExtensionProperty->GetValue(ObjectValue);
+
+		UCEClonerExtensionBase* ActiveExtension = Cast<UCEClonerExtensionBase>(ObjectValue);
+
+		if (!IsValid(ActiveExtension))
+		{
+			continue;
+		}
+
+		IDetailCategoryBuilder& ExtensionCategoryBuilder = InDetailBuilder.EditCategory(ActiveExtension->GetExtensionName());
+		ExtensionCategoryBuilder.SetSortOrder(StartOrder + ActiveExtension->GetExtensionSection().SectionOrder);
+
+		const FName ExtensionSectionName = ActiveExtension->GetExtensionSection().SectionName;
+		const TSharedRef<FPropertySection> ExtensionSection = PropertyModule.FindOrCreateSection(ComponentClassName, ExtensionSectionName, FText::FromName(ExtensionSectionName));
+		ExtensionSection->AddCategory(ActiveExtension->GetExtensionName());
+
+		TFunction<void(const TSharedPtr<IPropertyHandle>&)> AddObjectProperty = [&ExtensionCategoryBuilder, &AddObjectProperty](const TSharedPtr<IPropertyHandle>& InPropertyHandle)
+		{
+			uint32 NumChildren = 0;
+			InPropertyHandle->GetNumChildren(NumChildren);
+
+			for (uint32 ChildrenIndex = 0; ChildrenIndex < NumChildren; ChildrenIndex++)
+			{
+				TSharedPtr<IPropertyHandle> ChildHandle = InPropertyHandle->GetChildHandle(ChildrenIndex);
+
+				if (ChildHandle->IsCategoryHandle() && NumChildren == 1)
+				{
+					AddObjectProperty(ChildHandle);
+				}
+				else
+				{
+					ExtensionCategoryBuilder.AddProperty(ChildHandle);
+				}
+			}
+		};
+
+		AddObjectProperty(ActiveExtensionProperty->GetChildHandle(0));
+
+		if (UCEClonerLifetimeExtension* LifetimeExtension = Cast<UCEClonerLifetimeExtension>(ActiveExtension))
+		{
+			UNiagaraDataInterfaceCurve* LifetimeScaleCurve = LifetimeExtension->GetLifetimeScaleCurveDI();
+
+			if (!LifetimeScaleCurve || ClonerComponentsWeak.Num() > 1)
+			{
+				return;
+			}
+
+			const TArray<UObject*> ShowObjects {LifetimeScaleCurve};
+			ExtensionCategoryBuilder.SetShowAdvanced(true);
+
+			FAddPropertyParams Params;
+			Params.HideRootObjectNode(true);
+			Params.CreateCategoryNodes(false);
+
+			IDetailPropertyRow* CurveRow = ExtensionCategoryBuilder.AddExternalObjects(ShowObjects, EPropertyLocation::Advanced, Params);
+
+			TWeakObjectPtr<UCEClonerLifetimeExtension> LifetimeExtensionWeak(LifetimeExtension);
+			CurveRow->Visibility(MakeAttributeLambda([LifetimeExtensionWeak]()
+			{
+				if (UCEClonerLifetimeExtension* LifetimeExtension = LifetimeExtensionWeak.Get())
+				{
+					return LifetimeExtension->GetLifetimeEnabled()
+						&& LifetimeExtension->GetLifetimeScaleEnabled()
+						? EVisibility::Visible
+						: EVisibility::Collapsed;
+				}
+
+				return EVisibility::Visible;
+			}));
+		}
+	}
+
+	// Handle ufunctions
+	TMap<FName, FName> FunctionToCategory;
+	for (const TWeakObjectPtr<UCEClonerComponent>& ClonerComponentWeak : ClonerComponentsWeak)
+	{
+		UCEClonerComponent* ClonerComponent = ClonerComponentWeak.Get();
+
+		if (!ClonerComponent)
+		{
+			continue;
+		}
+
+		// Look for ufunction in component
+		for (UFunction* Function : TFieldRange<UFunction>(ClonerComponent->GetClass(), EFieldIteratorFlags::ExcludeSuper))
+		{
+			// Only CallInEditor function with 0 parameters
+			if (Function && Function->HasMetaData("CallInEditor") && Function->NumParms == 0)
+			{
+				FName FunctionName = Function->GetFName();
+				TMap<TWeakObjectPtr<UObject>, TWeakObjectPtr<UFunction>>& ObjectFunctions = LayoutFunctionNames.FindOrAdd(FunctionName);
+				ObjectFunctions.Add(ClonerComponent, Function);
+				FunctionToCategory.Add(FunctionName, FName(Function->GetMetaData(TEXT("Category"))));
+			}
+		}
+
+		// Look for CallInEditor functions in active layout
+		if (UCEClonerLayoutBase* ActiveLayout = ClonerComponent->GetActiveLayout())
+		{
+			// Iterate through all UFunctions in the class
+			for (UFunction* Function : TFieldRange<UFunction>(ActiveLayout->GetClass(), EFieldIteratorFlags::ExcludeSuper))
+			{
+				// Only CallInEditor function with 0 parameters
+				if (Function && Function->HasMetaData("CallInEditor") && Function->NumParms == 0)
+				{
+					FName FunctionName = Function->GetFName();
+					TMap<TWeakObjectPtr<UObject>, TWeakObjectPtr<UFunction>>& ObjectFunctions = LayoutFunctionNames.FindOrAdd(FunctionName);
+					ObjectFunctions.Add(ActiveLayout, Function);
+					FunctionToCategory.Add(FunctionName, FName(Function->GetMetaData(TEXT("Category"))));
+				}
+			}
+		}
+
+		// Look for CallInEditor functions in active extensions
+		for (const TObjectPtr<UCEClonerExtensionBase>& ActiveExtension : ClonerComponent->GetActiveExtensions())
+		{
+			// Iterate through all UFunctions in the class
+			for (UFunction* Function : TFieldRange<UFunction>(ActiveExtension->GetClass(), EFieldIteratorFlags::ExcludeSuper))
+			{
+				// Only CallInEditor function with 0 parameters
+				if (Function && Function->HasMetaData("CallInEditor") && Function->NumParms == 0)
+				{
+					FName FunctionName = Function->GetFName();
+					TMap<TWeakObjectPtr<UObject>, TWeakObjectPtr<UFunction>>& ObjectFunctions = LayoutFunctionNames.FindOrAdd(FunctionName);
+					ObjectFunctions.Add(ActiveExtension, Function);
+					FunctionToCategory.Add(FunctionName, FName(Function->GetMetaData(TEXT("Category"))));
+				}
+			}
+		}
+	}
+
+	if (!LayoutFunctionNames.IsEmpty())
+	{
+		TMap<FName, const TSharedPtr<SVerticalBox>> FunctionToWidget;
+
+		// Add buttons for ufunctions based on their category
+		for (const TPair<FName, FName>& FunctionToCategoryPair : FunctionToCategory)
+		{
+			const FName FunctionsCategoryName = FunctionToCategoryPair.Value;
+
+			const TSharedPtr<SVerticalBox> FunctionsWidget = SNew(SVerticalBox);
+
+			IDetailCategoryBuilder& FunctionsCategory = InDetailBuilder.EditCategory(FunctionsCategoryName, FText::FromName(FunctionsCategoryName), ECategoryPriority::Uncommon);
+
+			FunctionsCategory.AddCustomRow(FText::GetEmpty())
+				.WholeRowContent()
+				.HAlign(HAlign_Left)
+				[
+					FunctionsWidget.ToSharedRef()
+				];
+
+			const TSharedRef<FPropertySection> FunctionsSection = PropertyModule.FindOrCreateSection(ComponentClassName, FunctionsCategoryName, FText::FromName(FunctionsCategoryName));
+			FunctionsSection->AddCategory(FunctionsCategoryName);
+
+			FunctionToWidget.Add(FunctionToCategoryPair.Key, FunctionsWidget);
+		}
+
+		for (const TPair<FName, TMap<TWeakObjectPtr<UObject>, TWeakObjectPtr<UFunction>>>& LayoutFunctionNamesPair : LayoutFunctionNames)
+		{
+			const FText ButtonLabel = FText::FromString(FName::NameToDisplayString(LayoutFunctionNamesPair.Key.ToString(), false));
+
+			const TSharedPtr<SVerticalBox>& FunctionsWidget = FunctionToWidget.FindChecked(LayoutFunctionNamesPair.Key);
+
+			FunctionsWidget->AddSlot()
+				.Padding(0.f, 3.f)
+				.AutoHeight()
+				[
+					SNew(SButton)
+					.Text(ButtonLabel)
+					.HAlign(HAlign_Center)
+					.VAlign(VAlign_Fill)
+					.OnClicked(this, &FCEEditorClonerComponentDetailCustomization::OnFunctionButtonClicked, LayoutFunctionNamesPair.Key)
+				];
+		}
+	}
+}
+
+void FCEEditorClonerComponentDetailCustomization::RemoveEmptySections()
+{
+	FPropertyEditorModule& PropertyModule = FModuleManager::GetModuleChecked<FPropertyEditorModule>(TEXT("PropertyEditor"));
+	const FName ComponentClassName = UCEClonerComponent::StaticClass()->GetFName();
+
+	// Remove sections
+	PropertyModule.RemoveSection(ComponentClassName, TEXT("Rendering"));
+	PropertyModule.RemoveSection(ComponentClassName, TEXT("Effects"));
+	PropertyModule.RemoveSection(ComponentClassName, TEXT("Streaming"));
+}
+
+FReply FCEEditorClonerComponentDetailCustomization::OnFunctionButtonClicked(FName InFunctionName)
+{
+	if (TMap<TWeakObjectPtr<UObject>, TWeakObjectPtr<UFunction>> const* ObjectFunctions = LayoutFunctionNames.Find(InFunctionName))
+	{
+		for (const TPair<TWeakObjectPtr<UObject>, TWeakObjectPtr<UFunction>>& ObjectFunction : *ObjectFunctions)
+		{
+			UObject* Object = ObjectFunction.Key.Get();
+			UFunction* Function = ObjectFunction.Value.Get();
+
+			if (!Object || !Function)
+			{
+				continue;
+			}
+
+			Object->ProcessEvent(Function, nullptr);
+		}
+	}
+
+	return FReply::Handled();
+}
+
+#undef LOCTEXT_NAMESPACE
diff --git a/Engine/Plugins/Experimental/ClonerEffector/Source/ClonerEffectorEditor/Private/Cloner/Customizations/CEEditorClonerComponentDetailCustomization.h b/Engine/Plugins/Experimental/ClonerEffector/Source/ClonerEffectorEditor/Private/Cloner/Customizations/CEEditorClonerComponentDetailCustomization.h
new file mode 100644
--- /dev/null
+++ b/Engine/Plugins/Experimental/ClonerEffector/Source/ClonerEffectorEditor/Private/Cloner/Customizations/CEEditorClonerComponentDetailCustomization.h
@@ -0,0 +1,39 @@
+// Copyright Epic Games, Inc. All Rights Reserved.
+
+#pragma once
+
+#include "Containers/Map.h"
+#include "IDetailCustomization.h"
+#include "UObject/NameTypes.h"
+#include "UObject/WeakObjectPtr.h"
+
+class FReply;
+class UFunction;
+class UObject;
+
+/** Used to customize cloner component properties in details panel */
+class FCEEditorClonerComponentDetailCustomization : public IDetailCustomization
+{
+public:
+	static TSharedRef<IDetailCustomization> MakeInstance()
+	{
+		return MakeShared<FCEEditorClonerComponentDetailCustomization>();
+	}
+
+	explicit FCEEditorClonerComponentDetailCustomization()
+	{
+		RemoveEmptySections();
+	}
+
+	//~ Begin IDetailCustomization
+	virtual void CustomizeDetails(IDetailLayoutBuilder& InDetailBuilder) override;
+	//~ End IDetailCustomization
+
+protected:
+	static void RemoveEmptySections();
+
+	/** Execute ufunction with that name on selected objects */
+	FReply OnFunctionButtonClicked(FName InFunctionName);
+
+	TMap<FName, TMap<TWeakObjectPtr<UObject>, TWeakObjectPtr<UFunction>>> LayoutFunctionNames;
+};
diff --git a/Engine/Plugins/Experimental/ClonerEffector/Source/ClonerEffectorEditor/Private/Effector/CEEditorEffectorComponentDetailCustomization.cpp b/Engine/Plugins/Experimental/ClonerEffector/Source/ClonerEffectorEditor/Private/Effector/CEEditorEffectorComponentDetailCustomization.cpp
deleted file mode 100644
diff --git a/Engine/Plugins/Experimental/ClonerEffector/Source/ClonerEffectorEditor/Private/Effector/CEEditorEffectorComponentDetailCustomization.h b/Engine/Plugins/Experimental/ClonerEffector/Source/ClonerEffectorEditor/Private/Effector/CEEditorEffectorComponentDetailCustomization.h
deleted file mode 100644
diff --git a/Engine/Plugins/Experimental/ClonerEffector/Source/ClonerEffectorEditor/Private/Effector/CEEditorEffectorTypeDetailCustomization.cpp b/Engine/Plugins/Experimental/ClonerEffector/Source/ClonerEffectorEditor/Private/Effector/CEEditorEffectorTypeDetailCustomization.cpp
deleted file mode 100644
diff --git a/Engine/Plugins/Experimental/ClonerEffector/Source/ClonerEffectorEditor/Private/Effector/CEEditorEffectorTypeDetailCustomization.h b/Engine/Plugins/Experimental/ClonerEffector/Source/ClonerEffectorEditor/Private/Effector/CEEditorEffectorTypeDetailCustomization.h
deleted file mode 100644
diff --git a/Engine/Plugins/Experimental/ClonerEffector/Source/ClonerEffectorEditor/Private/Effector/Customizations/CEEditorEffectorActorDetailCustomization.cpp b/Engine/Plugins/Experimental/ClonerEffector/Source/ClonerEffectorEditor/Private/Effector/Customizations/CEEditorEffectorActorDetailCustomization.cpp
new file mode 100644
--- /dev/null
+++ b/Engine/Plugins/Experimental/ClonerEffector/Source/ClonerEffectorEditor/Private/Effector/Customizations/CEEditorEffectorActorDetailCustomization.cpp
@@ -0,0 +1,24 @@
+// Copyright Epic Games, Inc. All Rights Reserved.
+
+#include "Effector/Customizations/CEEditorEffectorActorDetailCustomization.h"
+
+#include "Effector/CEEffectorActor.h"
+#include "Modules/ModuleManager.h"
+#include "PropertyEditorModule.h"
+
+void FCEEditorEffectorActorDetailCustomization::CustomizeDetails(IDetailLayoutBuilder& InDetailBuilder)
+{
+	// Add customizations for actor here
+}
+
+void FCEEditorEffectorActorDetailCustomization::RemoveEmptySections()
+{
+	FPropertyEditorModule& PropertyModule = FModuleManager::GetModuleChecked<FPropertyEditorModule>("PropertyEditor");
+	const FName ComponentClassName = ACEEffectorActor::StaticClass()->GetFName();
+
+	// Remove Streaming section by removing its categories
+	TSharedRef<FPropertySection> StreamingSection = PropertyModule.FindOrCreateSection(ComponentClassName, "Streaming", FText::GetEmpty());
+	StreamingSection->RemoveCategory("WorldPartition");
+	StreamingSection->RemoveCategory("DataLayers");
+	StreamingSection->RemoveCategory("HLOD");
+}
diff --git a/Engine/Plugins/Experimental/ClonerEffector/Source/ClonerEffectorEditor/Private/Effector/Customizations/CEEditorEffectorActorDetailCustomization.h b/Engine/Plugins/Experimental/ClonerEffector/Source/ClonerEffectorEditor/Private/Effector/Customizations/CEEditorEffectorActorDetailCustomization.h
new file mode 100644
--- /dev/null
+++ b/Engine/Plugins/Experimental/ClonerEffector/Source/ClonerEffectorEditor/Private/Effector/Customizations/CEEditorEffectorActorDetailCustomization.h
@@ -0,0 +1,28 @@
+// Copyright Epic Games, Inc. All Rights Reserved.
+
+#pragma once
+
+#include "IDetailCustomization.h"
+#include "Templates/SharedPointer.h"
+
+/** Used to customize effector actor properties in details panel */
+class FCEEditorEffectorActorDetailCustomization : public IDetailCustomization
+{
+public:
+	static TSharedRef<IDetailCustomization> MakeInstance()
+	{
+		return MakeShared<FCEEditorEffectorActorDetailCustomization>();
+	}
+
+	explicit FCEEditorEffectorActorDetailCustomization()
+	{
+		RemoveEmptySections();
+	}
+
+	//~ Begin IDetailCustomization
+	virtual void CustomizeDetails(IDetailLayoutBuilder& InDetailBuilder) override;
+	//~ End IDetailCustomization
+
+protected:
+	static void RemoveEmptySections();
+};
diff --git a/Engine/Plugins/Experimental/ClonerEffector/Source/ClonerEffectorEditor/Private/Effector/Customizations/CEEditorEffectorComponentDetailCustomization.cpp b/Engine/Plugins/Experimental/ClonerEffector/Source/ClonerEffectorEditor/Private/Effector/Customizations/CEEditorEffectorComponentDetailCustomization.cpp
new file mode 100644
--- /dev/null
+++ b/Engine/Plugins/Experimental/ClonerEffector/Source/ClonerEffectorEditor/Private/Effector/Customizations/CEEditorEffectorComponentDetailCustomization.cpp
@@ -0,0 +1,135 @@
+// Copyright Epic Games, Inc. All Rights Reserved.
+
+#include "Effector/Customizations/CEEditorEffectorComponentDetailCustomization.h"
+
+#include "DetailLayoutBuilder.h"
+#include "DetailWidgetRow.h"
+#include "Effector/CEEffectorComponent.h"
+#include "Effector/CEEffectorExtensionBase.h"
+#include "Modules/ModuleManager.h"
+#include "PropertyEditorModule.h"
+#include "Widgets/Input/SButton.h"
+
+#define LOCTEXT_NAMESPACE "CEEditorEffectorComponentDetailCustomization"
+
+void FCEEditorEffectorComponentDetailCustomization::CustomizeDetails(IDetailLayoutBuilder& InDetailBuilder)
+{
+	FPropertyEditorModule& PropertyModule = FModuleManager::GetModuleChecked<FPropertyEditorModule>(TEXT("PropertyEditor"));
+	const FName ComponentClassName = UCEEffectorComponent::StaticClass()->GetFName();
+
+	// Remove extension array property
+	TSharedRef<IPropertyHandle> ActiveExtensionsProperty = InDetailBuilder.GetProperty(GET_MEMBER_NAME_CHECKED(UCEEffectorComponent, ActiveExtensions), UCEEffectorComponent::StaticClass());
+
+	if (!ActiveExtensionsProperty->IsValidHandle())
+	{
+		return;
+	}
+
+	InDetailBuilder.HideProperty(ActiveExtensionsProperty);
+
+	// Everything needs to be below Effector category
+	const FName EffectorCategoryName = TEXT("Effector");
+	int32 StartOrder = InDetailBuilder.EditCategory(EffectorCategoryName).GetSortOrder() + 1;
+
+	const FName EffectorSectionName = UE::ClonerEffector::EffectorSection::EffectorSection.SectionName;
+	const TSharedRef<FPropertySection> EffectorSection = PropertyModule.FindOrCreateSection(ComponentClassName, EffectorSectionName, FText::FromName(EffectorSectionName));
+	EffectorSection->AddCategory(EffectorCategoryName);
+
+	// Shape
+	{
+		TSharedRef<IPropertyHandle> ActiveTypeProperty = InDetailBuilder.GetProperty(GET_MEMBER_NAME_CHECKED(UCEEffectorComponent, ActiveType), UCEEffectorComponent::StaticClass());
+
+		if (!ActiveTypeProperty->IsValidHandle())
+		{
+			return;
+		}
+
+		const FName TypeCategoryName = ActiveTypeProperty->GetDefaultCategoryName();
+
+		IDetailCategoryBuilder& TypeCategory = InDetailBuilder.EditCategory(TypeCategoryName);
+		TypeCategory.SetSortOrder(StartOrder++);
+
+		const FName TypeSectionName = UE::ClonerEffector::EffectorSection::ShapeSection.SectionName;
+		const TSharedRef<FPropertySection> TypeSection = PropertyModule.FindOrCreateSection(ComponentClassName, TypeSectionName, FText::FromName(TypeSectionName));
+		TypeSection->AddCategory(TypeCategoryName);
+	}
+
+	// Mode
+	{
+		TSharedRef<IPropertyHandle> ActiveModeProperty = InDetailBuilder.GetProperty(GET_MEMBER_NAME_CHECKED(UCEEffectorComponent, ActiveMode), UCEEffectorComponent::StaticClass());
+
+		if (!ActiveModeProperty->IsValidHandle())
+		{
+			return;
+		}
+
+		const FName ModeCategoryName = ActiveModeProperty->GetDefaultCategoryName();
+
+		IDetailCategoryBuilder& ModeCategory = InDetailBuilder.EditCategory(ModeCategoryName);
+		ModeCategory.SetSortOrder(StartOrder++);
+
+		const FName ModeSectionName = UE::ClonerEffector::EffectorSection::ModeSection.SectionName;
+		const TSharedRef<FPropertySection> ModeSection = PropertyModule.FindOrCreateSection(ComponentClassName, ModeSectionName, FText::FromName(ModeSectionName));
+		ModeSection->AddCategory(ModeCategoryName);
+	}
+
+	TFunction<void(const TSharedPtr<IPropertyHandle>&, IDetailCategoryBuilder&)> AddObjectProperty = [&AddObjectProperty](const TSharedPtr<IPropertyHandle>& InPropertyHandle, IDetailCategoryBuilder& InCategoryBuilder)
+	{
+		uint32 NumChildren = 0;
+		InPropertyHandle->GetNumChildren(NumChildren);
+
+		for (uint32 ChildrenIndex = 0; ChildrenIndex < NumChildren; ChildrenIndex++)
+		{
+			TSharedPtr<IPropertyHandle> ChildHandle = InPropertyHandle->GetChildHandle(ChildrenIndex);
+
+			if (ChildHandle->IsCategoryHandle() && NumChildren == 1)
+			{
+				AddObjectProperty(ChildHandle, InCategoryBuilder);
+			}
+			else
+			{
+				InCategoryBuilder.AddProperty(ChildHandle);
+			}
+		}
+	};
+
+	uint32 NumElements = 0;
+	TSharedPtr<IPropertyHandleArray> ActiveExtensionArrayProperty = ActiveExtensionsProperty->AsArray();
+	ActiveExtensionArrayProperty->GetNumElements(NumElements);
+
+	for (uint32 ElementIndex = 0; ElementIndex < NumElements; ElementIndex++)
+	{
+		TSharedPtr<IPropertyHandle> ActiveExtensionProperty = ActiveExtensionArrayProperty->GetElement(ElementIndex);
+
+		UObject* ObjectValue;
+		ActiveExtensionProperty->GetValue(ObjectValue);
+
+		UCEEffectorExtensionBase* ActiveExtension = Cast<UCEEffectorExtensionBase>(ObjectValue);
+
+		if (!ActiveExtension)
+		{
+			continue;
+		}
+
+		const FName ExtensionSectionName = ActiveExtension->GetExtensionSection().SectionName;
+
+		IDetailCategoryBuilder& ExtensionCategory = InDetailBuilder.EditCategory(ExtensionSectionName);
+		ExtensionCategory.SetSortOrder(StartOrder + ActiveExtension->GetExtensionSection().SectionOrder);
+
+		const TSharedRef<FPropertySection> ExtensionSection = PropertyModule.FindOrCreateSection(ComponentClassName, ExtensionSectionName, FText::FromName(ExtensionSectionName));
+		ExtensionSection->AddCategory(ExtensionSectionName);
+
+		AddObjectProperty(ActiveExtensionProperty->GetChildHandle(0), ExtensionCategory);
+	}
+}
+
+void FCEEditorEffectorComponentDetailCustomization::RemoveEmptySections()
+{
+	FPropertyEditorModule& PropertyModule = FModuleManager::GetModuleChecked<FPropertyEditorModule>(TEXT("PropertyEditor"));
+	const FName ComponentClassName = UCEEffectorComponent::StaticClass()->GetFName();
+
+	// Remove sections
+	PropertyModule.RemoveSection(ComponentClassName, TEXT("Streaming"));
+}
+
+#undef LOCTEXT_NAMESPACE
diff --git a/Engine/Plugins/Experimental/ClonerEffector/Source/ClonerEffectorEditor/Private/Effector/Customizations/CEEditorEffectorComponentDetailCustomization.h b/Engine/Plugins/Experimental/ClonerEffector/Source/ClonerEffectorEditor/Private/Effector/Customizations/CEEditorEffectorComponentDetailCustomization.h
new file mode 100644
--- /dev/null
+++ b/Engine/Plugins/Experimental/ClonerEffector/Source/ClonerEffectorEditor/Private/Effector/Customizations/CEEditorEffectorComponentDetailCustomization.h
@@ -0,0 +1,27 @@
+// Copyright Epic Games, Inc. All Rights Reserved.
+
+#pragma once
+
+#include "IDetailCustomization.h"
+
+/** Used to customize effector component properties in details panel */
+class FCEEditorEffectorComponentDetailCustomization : public IDetailCustomization
+{
+public:
+	static TSharedRef<IDetailCustomization> MakeInstance()
+	{
+		return MakeShared<FCEEditorEffectorComponentDetailCustomization>();
+	}
+
+	explicit FCEEditorEffectorComponentDetailCustomization()
+	{
+		RemoveEmptySections();
+	}
+
+	//~ Begin IDetailCustomization
+	virtual void CustomizeDetails(IDetailLayoutBuilder& InDetailBuilder) override;
+	//~ End IDetailCustomization
+
+protected:
+	static void RemoveEmptySections();
+};
diff --git a/Engine/Plugins/Experimental/ClonerEffector/Source/ClonerEffectorEditor/Private/Effector/Customizations/CEEditorEffectorTypeDetailCustomization.cpp b/Engine/Plugins/Experimental/ClonerEffector/Source/ClonerEffectorEditor/Private/Effector/Customizations/CEEditorEffectorTypeDetailCustomization.cpp
new file mode 100644
--- /dev/null
+++ b/Engine/Plugins/Experimental/ClonerEffector/Source/ClonerEffectorEditor/Private/Effector/Customizations/CEEditorEffectorTypeDetailCustomization.cpp
@@ -0,0 +1,267 @@
+// Copyright Epic Games, Inc. All Rights Reserved.
+
+#include "Effector/Customizations/CEEditorEffectorTypeDetailCustomization.h"
+
+#include "DetailLayoutBuilder.h"
+#include "DetailWidgetRow.h"
+#include "Effector/Types/CEEffectorBoundType.h"
+#include "Styles/CEEditorStyle.h"
+#include "Widgets/Layout/SBox.h"
+#include "Widgets/Layout/SScaleBox.h"
+#include "Widgets/Text/STextBlock.h"
+#include "Widgets/Images/SImage.h"
+#include "Widgets/Input/SComboBox.h"
+
+DEFINE_LOG_CATEGORY_STATIC(LogCEEditorEffectorTypeDetailCustomization, Log, All);
+
+#define LOCTEXT_NAMESPACE "CEEditorEffectorTypeDetailCustomization"
+
+namespace UE::EffectorType
+{
+	// Sort from most dramatic to least IN then OUT then IN OUT, then specials
+	static const TArray<ECEClonerEasing> SortedEasings
+	{
+		ECEClonerEasing::InExpo,
+		ECEClonerEasing::InCirc,
+		ECEClonerEasing::InQuint,
+		ECEClonerEasing::InQuart,
+		ECEClonerEasing::InQuad,
+		ECEClonerEasing::InCubic,
+		ECEClonerEasing::InSine,
+		ECEClonerEasing::OutExpo,
+		ECEClonerEasing::OutCirc,
+		ECEClonerEasing::OutQuint,
+		ECEClonerEasing::OutQuart,
+		ECEClonerEasing::OutQuad,
+		ECEClonerEasing::OutCubic,
+		ECEClonerEasing::OutSine,
+		ECEClonerEasing::InOutExpo,
+		ECEClonerEasing::InOutCirc,
+		ECEClonerEasing::InOutQuint,
+		ECEClonerEasing::InOutQuart,
+		ECEClonerEasing::InOutQuad,
+		ECEClonerEasing::InOutCubic,
+		ECEClonerEasing::InOutSine,
+		ECEClonerEasing::Linear,
+		ECEClonerEasing::InBounce,
+		ECEClonerEasing::InBack,
+		ECEClonerEasing::InElastic,
+		ECEClonerEasing::OutBounce,
+		ECEClonerEasing::OutBack,
+		ECEClonerEasing::OutElastic,
+		ECEClonerEasing::InOutBounce,
+		ECEClonerEasing::InOutBack,
+		ECEClonerEasing::InOutElastic,
+		ECEClonerEasing::Random
+	};
+}
+
+void FCEEditorEffectorTypeDetailCustomization::CustomizeDetails(IDetailLayoutBuilder& InDetailBuilder)
+{
+	EasingPropertyHandle = InDetailBuilder.GetProperty(
+		GET_MEMBER_NAME_CHECKED(UCEEffectorBoundType, Easing),
+		UCEEffectorBoundType::StaticClass()
+	);
+
+	// Customize easing curve property
+	if (!EasingPropertyHandle->IsValidHandle())
+	{
+		return;
+	}
+
+	IDetailPropertyRow& EasingRow = InDetailBuilder.AddPropertyToCategory(EasingPropertyHandle);
+
+	PopulateEasingInfos();
+
+	FDetailWidgetRow& CustomWidget = EasingRow.CustomWidget();
+
+	CustomWidget.NameContent()
+		.VAlign(EVerticalAlignment::VAlign_Fill)
+		[
+			EasingPropertyHandle->CreatePropertyNameWidget()
+		];
+
+	CustomWidget.ValueContent()
+		.VAlign(EVerticalAlignment::VAlign_Fill)
+		[
+			SNew(SComboBox<FName>)
+			.ComboBoxStyle(&FCEEditorStyle::Get().GetWidgetStyle<FComboBoxStyle>(TEXT("ComboBox")))
+			.OptionsSource(&EasingNames)
+			.InitiallySelectedItem(GetCurrentEasingName())
+			.ToolTipText(LOCTEXT("EasingTooltip", "Easings sorted from most dramatic to least and specials at the end"))
+			.OnGenerateWidget(this, &FCEEditorEffectorTypeDetailCustomization::OnGenerateEasingEntry)
+			.OnSelectionChanged(this, &FCEEditorEffectorTypeDetailCustomization::OnSelectionChanged)
+			.ContentPadding(0.f)
+			.Content()
+			[
+				OnGenerateEasingEntry(NAME_None)
+			]
+		];
+}
+
+void FCEEditorEffectorTypeDetailCustomization::PopulateEasingInfos()
+{
+	UEnum* EasingEnum = StaticEnum<ECEClonerEasing>();
+
+	if (!EasingEnum)
+	{
+		return;
+	}
+
+	EasingEnumWeak = EasingEnum;
+
+	check((EasingEnum->NumEnums()-1) == UE::EffectorType::SortedEasings.Num());
+
+	EasingNames.Empty(UE::EffectorType::SortedEasings.Num());
+	for (const ECEClonerEasing& Easing : UE::EffectorType::SortedEasings)
+	{
+		EasingNames.Add(FName(EasingEnum->GetNameStringByValue(static_cast<uint8>(Easing))));
+	}
+}
+
+TSharedRef<SWidget> FCEEditorEffectorTypeDetailCustomization::OnGenerateEasingEntry(FName InName) const
+{
+	TSharedPtr<SWidget> ImageWidget;
+	TSharedPtr<SWidget> TextWidget;
+	const TSharedRef<SHorizontalBox> HorizontalWidget = SNew(SHorizontalBox).Visibility(EVisibility::Visible);
+
+	static const FVector2D ImageSizeClosed(16.f, 16.f);
+	static const FVector2D ImageSizeOpened(32.f, 32.f);
+
+	// If none = update with current value else set fixed value once
+	if (InName == NAME_None)
+	{
+		SAssignNew(ImageWidget, SImage)
+			.ColorAndOpacity(FAppStyle::GetSlateColor("SelectionColor"))
+			.DesiredSizeOverride(ImageSizeClosed)
+			.Image(this, &FCEEditorEffectorTypeDetailCustomization::GetEasingImage, InName);
+
+		SAssignNew(TextWidget, STextBlock)
+			.Justification(ETextJustify::Center)
+			.Text(this, &FCEEditorEffectorTypeDetailCustomization::GetEasingText, InName);
+	}
+	else
+	{
+		// We need to switch color on hover to avoid having the selected color equals the image color
+		const TSharedPtr<SWidget> HoverWidget = HorizontalWidget;
+
+		SAssignNew(ImageWidget, SImage)
+			.ColorAndOpacity_Static(&FCEEditorEffectorTypeDetailCustomization::GetImageColorAndOpacity, HoverWidget)
+			.DesiredSizeOverride(ImageSizeOpened)
+			.Image(GetEasingImage(InName));
+
+		SAssignNew(TextWidget, STextBlock)
+			.Justification(ETextJustify::Center)
+			.Text(GetEasingText(InName));
+	}
+
+	// Make inner widgets hit test invisible and only horizontal is hit testable
+
+	HorizontalWidget->AddSlot()
+		.AutoWidth()
+		.Padding(2.f)
+		[
+			SNew(SScaleBox)
+			.Visibility(EVisibility::HitTestInvisible)
+			.Stretch(EStretch::UserSpecified)
+			.UserSpecifiedScale(1.5)
+			[
+				ImageWidget.ToSharedRef()
+			]
+		];
+
+	HorizontalWidget->AddSlot()
+		.FillWidth(1.f)
+		.Padding(8.f, 2.f)
+		.HAlign(EHorizontalAlignment::HAlign_Fill)
+		.VAlign(EVerticalAlignment::VAlign_Center)
+		[
+			SNew(SBox)
+			.Visibility(EVisibility::HitTestInvisible)
+			[
+				TextWidget.ToSharedRef()
+			]
+		];
+
+	return HorizontalWidget;
+}
+
+void FCEEditorEffectorTypeDetailCustomization::OnSelectionChanged(FName InSelection, ESelectInfo::Type InSelectInfo) const
+{
+	UEnum* EasingEnum = EasingEnumWeak.Get();
+
+	if (EasingEnum && EasingPropertyHandle.IsValid())
+	{
+		const int32 Idx = EasingEnum->GetValueByNameString(InSelection.ToString());
+
+		if (Idx != INDEX_NONE && EasingPropertyHandle->SetValue(static_cast<uint8>(Idx), EPropertyValueSetFlags::DefaultFlags) != FPropertyAccess::Success)
+		{
+			UE_LOG(LogCEEditorEffectorTypeDetailCustomization, Warning, TEXT("EffectorTypeDetailCustomization : Failed to set property value %s on selection"), *InSelection.ToString())
+		}
+	}
+}
+
+FName FCEEditorEffectorTypeDetailCustomization::GetCurrentEasingName() const
+{
+	uint8 CurrentValue;
+
+	if (EasingPropertyHandle.IsValid() && EasingPropertyHandle->GetValue(CurrentValue) == FPropertyAccess::Result::Success)
+	{
+		if (UEnum* EasingEnum = EasingEnumWeak.Get())
+		{
+			return FName(EasingEnum->GetNameStringByValue(CurrentValue));
+		}
+	}
+
+	return NAME_None;
+}
+
+const FSlateBrush* FCEEditorEffectorTypeDetailCustomization::GetEasingImage(FName InName) const
+{
+	if (InName == NAME_None)
+	{
+		InName = GetCurrentEasingName();
+
+		// multiple values
+		if (InName == NAME_None)
+		{
+			return nullptr;
+		}
+	}
+
+	return FCEEditorStyle::Get().GetBrush(FName(TEXT("EasingIcons.") + InName.ToString()));
+}
+
+FText FCEEditorEffectorTypeDetailCustomization::GetEasingText(FName InName) const
+{
+	if (InName == NAME_None)
+	{
+		InName = GetCurrentEasingName();
+
+		// multiple values
+		if (InName == NAME_None)
+		{
+			return LOCTEXT("MultipleValue", "Multiple values selected");
+		}
+	}
+
+	if (UEnum* EasingEnum = EasingEnumWeak.Get())
+	{
+		const int32 EnumValue = EasingEnum->GetValueByNameString(InName.ToString());
+		return EasingEnum->GetDisplayNameTextByValue(EnumValue);
+	}
+
+	return FText::GetEmpty();
+}
+
+FSlateColor FCEEditorEffectorTypeDetailCustomization::GetImageColorAndOpacity(const TSharedPtr<SWidget> InWidget)
+{
+	if (InWidget.IsValid())
+	{
+		return InWidget->IsHovered() ? FLinearColor::White : FAppStyle::GetSlateColor("SelectionColor");
+	}
+
+	return FLinearColor::White;
+}
+
+#undef LOCTEXT_NAMESPACE
\ No newline at end of file
diff --git a/Engine/Plugins/Experimental/ClonerEffector/Source/ClonerEffectorEditor/Private/Effector/Customizations/CEEditorEffectorTypeDetailCustomization.h b/Engine/Plugins/Experimental/ClonerEffector/Source/ClonerEffectorEditor/Private/Effector/Customizations/CEEditorEffectorTypeDetailCustomization.h
new file mode 100644
--- /dev/null
+++ b/Engine/Plugins/Experimental/ClonerEffector/Source/ClonerEffectorEditor/Private/Effector/Customizations/CEEditorEffectorTypeDetailCustomization.h
@@ -0,0 +1,44 @@
+// Copyright Epic Games, Inc. All Rights Reserved.
+
+#pragma once
+
+#include "IDetailCustomization.h"
+#include "Templates/SharedPointer.h"
+#include "Types/SlateEnums.h"
+#include "UObject/WeakObjectPtr.h"
+
+class IPropertyHandle;
+class SWidget;
+class UEnum;
+struct FLinearColor;
+struct FSlateBrush;
+struct FSlateColor;
+
+/** Used to customize effector type properties in details panel */
+class FCEEditorEffectorTypeDetailCustomization : public IDetailCustomization
+{
+public:
+	static TSharedRef<IDetailCustomization> MakeInstance()
+	{
+		return MakeShared<FCEEditorEffectorTypeDetailCustomization>();
+	}
+
+	//~ Begin IDetailCustomization
+	virtual void CustomizeDetails(IDetailLayoutBuilder& InDetailBuilder) override;
+	//~ End IDetailCustomization
+
+protected:
+	static FSlateColor GetImageColorAndOpacity(const TSharedPtr<SWidget> InWidget);
+
+	void PopulateEasingInfos();
+	FName GetCurrentEasingName() const;
+	const FSlateBrush* GetEasingImage(FName InName) const;
+	FText GetEasingText(FName InName) const;
+
+	TSharedRef<SWidget> OnGenerateEasingEntry(FName InName) const;
+	void OnSelectionChanged(FName InSelection, ESelectInfo::Type InSelectInfo) const;
+
+	TArray<FName> EasingNames;
+	TWeakObjectPtr<UEnum> EasingEnumWeak;
+	TSharedPtr<IPropertyHandle> EasingPropertyHandle;
+};
