MotionDesign : ClonerEffector
- Adding outliner context menu options for cloner and effector (enable, disable, convert, spawn effector)

diff --git a/Engine/Plugins/Experimental/Avalanche/Source/AvalancheEffectorsEditor/AvalancheEffectorsEditor.Build.cs b/Engine/Plugins/Experimental/Avalanche/Source/AvalancheEffectorsEditor/AvalancheEffectorsEditor.Build.cs
--- a/Engine/Plugins/Experimental/Avalanche/Source/AvalancheEffectorsEditor/AvalancheEffectorsEditor.Build.cs
+++ b/Engine/Plugins/Experimental/Avalanche/Source/AvalancheEffectorsEditor/AvalancheEffectorsEditor.Build.cs
@@ -22,6 +22,7 @@ public class AvalancheEffectorsEditor : ModuleRules
 				"AvalancheCore",
 				"AvalancheEffectors",
 				"AvalancheInteractiveTools",
+				"AvalancheOutliner",
 				"AvalancheShapes",
 				"ClonerEffector",
 				"ClonerEffectorEditor",
diff --git a/Engine/Plugins/Experimental/Avalanche/Source/AvalancheEffectorsEditor/Private/AvalancheEffectorsEditorModule.cpp b/Engine/Plugins/Experimental/Avalanche/Source/AvalancheEffectorsEditor/Private/AvalancheEffectorsEditorModule.cpp
--- a/Engine/Plugins/Experimental/Avalanche/Source/AvalancheEffectorsEditor/Private/AvalancheEffectorsEditorModule.cpp
+++ b/Engine/Plugins/Experimental/Avalanche/Source/AvalancheEffectorsEditor/Private/AvalancheEffectorsEditorModule.cpp
@@ -6,13 +6,16 @@
 #include "AvaInteractiveToolsDelegates.h"
 #include "Cloner/AvaClonerActorTool.h"
 #include "Cloner/AvaClonerActorVis.h"
+#include "Cloner/AvaClonerEditorOutlinerContextMenu.h"
 #include "Cloner/CEClonerComponent.h"
 #include "ComponentVisualizers.h"
 #include "Effector/AvaEffectorActorTool.h"
 #include "Effector/AvaEffectorActorVis.h"
+#include "Effector/AvaEffectorEditorOutlinerContextMenu.h"
 #include "Effector/CEEffectorComponent.h"
 #include "Framework/Application/SlateApplication.h"
 #include "IAvalancheComponentVisualizersModule.h"
+#include "IAvaOutlinerModule.h"
 #include "Modules/ModuleManager.h"
 #include "PropertyEditorModule.h"
 
@@ -23,6 +26,8 @@ void FAvalancheEffectorsEditorModule::StartupModule()
 	FAvaEffectorsEditorCommands::Register();
 
 	FAvaInteractiveToolsDelegates::GetRegisterToolsDelegate().AddRaw(this, &FAvalancheEffectorsEditorModule::RegisterTools);
+
+	RegisterOutlinerItems();
 }
 
 void FAvalancheEffectorsEditorModule::ShutdownModule()
@@ -32,6 +37,8 @@ void FAvalancheEffectorsEditorModule::ShutdownModule()
 	FAvaEffectorsEditorCommands::Unregister();
 
 	FAvaInteractiveToolsDelegates::GetRegisterToolsDelegate().RemoveAll(this);
+
+	UnregisterOutlinerItems();
 }
 
 void FAvalancheEffectorsEditorModule::PostEngineInit()
@@ -61,4 +68,29 @@ void FAvalancheEffectorsEditorModule::RegisterComponentVisualizers()
 	IAvalancheComponentVisualizersModule::RegisterComponentVisualizer<UCEClonerComponent, FAvaClonerActorVisualizer>(&Visualizers);
 }
 
+void FAvalancheEffectorsEditorModule::RegisterOutlinerItems()
+{
+	IAvaOutlinerModule& OutlinerModule = IAvaOutlinerModule::Get();
+
+	OutlinerContextClonerDelegateHandle = OutlinerModule.GetOnExtendOutlinerItemContextMenu()
+		.AddStatic(&FAvaClonerEditorOutlinerContextMenu::OnExtendOutlinerContextMenu);
+
+	OutlinerContextEffectorDelegateHandle = OutlinerModule.GetOnExtendOutlinerItemContextMenu()
+		.AddStatic(&FAvaEffectorEditorOutlinerContextMenu::OnExtendOutlinerContextMenu);
+}
+
+void FAvalancheEffectorsEditorModule::UnregisterOutlinerItems()
+{
+	if (IAvaOutlinerModule::IsLoaded())
+	{
+		IAvaOutlinerModule& OutlinerModule = IAvaOutlinerModule::Get();
+
+		OutlinerModule.GetOnExtendOutlinerItemContextMenu().Remove(OutlinerContextClonerDelegateHandle);
+        OutlinerContextClonerDelegateHandle.Reset();
+
+		OutlinerModule.GetOnExtendOutlinerItemContextMenu().Remove(OutlinerContextEffectorDelegateHandle);
+		OutlinerContextEffectorDelegateHandle.Reset();
+	}
+}
+
 IMPLEMENT_MODULE(FAvalancheEffectorsEditorModule, AvalancheEffectorsEditor)
diff --git a/Engine/Plugins/Experimental/Avalanche/Source/AvalancheEffectorsEditor/Private/AvalancheEffectorsEditorModule.h b/Engine/Plugins/Experimental/Avalanche/Source/AvalancheEffectorsEditor/Private/AvalancheEffectorsEditorModule.h
--- a/Engine/Plugins/Experimental/Avalanche/Source/AvalancheEffectorsEditor/Private/AvalancheEffectorsEditorModule.h
+++ b/Engine/Plugins/Experimental/Avalanche/Source/AvalancheEffectorsEditor/Private/AvalancheEffectorsEditorModule.h
@@ -3,6 +3,7 @@
 #pragma once
 
 #include "Containers/Array.h"
+#include "Delegates/IDelegateInstance.h"
 #include "Modules/ModuleInterface.h"
 #include "Templates/SharedPointerFwd.h"
 
@@ -24,5 +25,11 @@ private:
 
 	void RegisterComponentVisualizers();
 
+	void RegisterOutlinerItems();
+	void UnregisterOutlinerItems();
+
 	TArray<TSharedPtr<FComponentVisualizer>> Visualizers;
+
+	FDelegateHandle OutlinerContextClonerDelegateHandle;
+	FDelegateHandle OutlinerContextEffectorDelegateHandle;
 };
diff --git a/Engine/Plugins/Experimental/Avalanche/Source/AvalancheEffectorsEditor/Private/Cloner/AvaClonerEditorOutlinerContextMenu.cpp b/Engine/Plugins/Experimental/Avalanche/Source/AvalancheEffectorsEditor/Private/Cloner/AvaClonerEditorOutlinerContextMenu.cpp
new file mode 100644
--- /dev/null
+++ b/Engine/Plugins/Experimental/Avalanche/Source/AvalancheEffectorsEditor/Private/Cloner/AvaClonerEditorOutlinerContextMenu.cpp
@@ -0,0 +1,94 @@
+// Copyright Epic Games, Inc. All Rights Reserved.
+
+#include "Cloner/AvaClonerEditorOutlinerContextMenu.h"
+
+#include "Cloner/CEClonerComponent.h"
+#include "Cloner/Menus/CEEditorClonerMenuContext.h"
+#include "Cloner/Menus/CEEditorClonerMenuOptions.h"
+#include "GameFramework/Actor.h"
+#include "Item/AvaOutlinerActor.h"
+#include "Item/AvaOutlinerComponent.h"
+#include "Subsystems/CEEditorClonerSubsystem.h"
+#include "ToolMenuContext/AvaOutlinerItemsContext.h"
+
+void FAvaClonerEditorOutlinerContextMenu::OnExtendOutlinerContextMenu(UToolMenu* InToolMenu)
+{
+	if (!InToolMenu)
+	{
+		return;
+	}
+
+	const UAvaOutlinerItemsContext* OutlinerItemsContext = InToolMenu->Context.FindContext<UAvaOutlinerItemsContext>();
+
+	if (!IsValid(OutlinerItemsContext))
+	{
+		return;
+	}
+
+	UCEEditorClonerSubsystem* ClonerSubsystem = UCEEditorClonerSubsystem::Get();
+
+	if (!IsValid(ClonerSubsystem))
+	{
+		return;
+	}
+
+	TSet<UObject*> ContextObjects;
+	GetContextObjects(OutlinerItemsContext, ContextObjects);
+
+	if (ContextObjects.IsEmpty())
+	{
+		return;
+	}
+
+	const FCEEditorClonerMenuContext Context(ContextObjects);
+	FCEEditorClonerMenuOptions Options(
+		{
+			ECEEditorClonerMenuType::Enable
+			, ECEEditorClonerMenuType::Disable
+			, ECEEditorClonerMenuType::CreateEffector
+			, ECEEditorClonerMenuType::Convert
+		}
+	);
+	Options.UseTransact(true);
+	Options.CreateSubMenu(true);
+
+	ClonerSubsystem->FillClonerMenu(InToolMenu, Context, Options);
+}
+
+void FAvaClonerEditorOutlinerContextMenu::GetContextObjects(const UAvaOutlinerItemsContext* InContext, TSet<UObject*>& OutObjects)
+{
+	OutObjects.Empty();
+
+	if (!IsValid(InContext) || InContext->GetItems().IsEmpty())
+	{
+		return;
+	}
+
+	for (const FAvaOutlinerItemWeakPtr& ItemWeak : InContext->GetItems())
+	{
+		const TSharedPtr<IAvaOutlinerItem> Item = ItemWeak.Pin();
+
+		// is it an actor
+		if (const FAvaOutlinerActor* ActorItem = Item->CastTo<FAvaOutlinerActor>())
+		{
+			if (AActor* Actor = ActorItem->GetActor())
+			{
+				if (Actor->FindComponentByClass<UCEClonerComponent>())
+				{
+					OutObjects.Add(Actor);
+				}
+			}
+		}
+		// is it a component
+		else if (const FAvaOutlinerComponent* ComponentItem = Item->CastTo<FAvaOutlinerComponent>())
+		{
+			if (USceneComponent* Component = ComponentItem->GetComponent())
+			{
+				if (Component->IsA<UCEClonerComponent>())
+				{
+					OutObjects.Add(Component);
+				}
+			}
+		}
+	}
+}
diff --git a/Engine/Plugins/Experimental/Avalanche/Source/AvalancheEffectorsEditor/Private/Cloner/AvaClonerEditorOutlinerContextMenu.h b/Engine/Plugins/Experimental/Avalanche/Source/AvalancheEffectorsEditor/Private/Cloner/AvaClonerEditorOutlinerContextMenu.h
new file mode 100644
--- /dev/null
+++ b/Engine/Plugins/Experimental/Avalanche/Source/AvalancheEffectorsEditor/Private/Cloner/AvaClonerEditorOutlinerContextMenu.h
@@ -0,0 +1,21 @@
+// Copyright Epic Games, Inc. All Rights Reserved.
+
+#pragma once
+
+#include "Containers/Set.h"
+#include "ToolMenus.h"
+#include "UObject/Object.h"
+#include "UObject/WeakObjectPtrTemplates.h"
+
+class UAvaOutlinerItemsContext;
+
+/** Context menu actions in outliner related to cloner */
+class FAvaClonerEditorOutlinerContextMenu
+{
+public:
+	/** Extend outliner context menu */
+	static void OnExtendOutlinerContextMenu(UToolMenu* InToolMenu);
+
+	/** Gets all pertinent context object in this context */
+	static void GetContextObjects(const UAvaOutlinerItemsContext* InContext, TSet<UObject*>& OutObjects);
+};
\ No newline at end of file
diff --git a/Engine/Plugins/Experimental/Avalanche/Source/AvalancheEffectorsEditor/Private/Effector/AvaEffectorEditorOutlinerContextMenu.cpp b/Engine/Plugins/Experimental/Avalanche/Source/AvalancheEffectorsEditor/Private/Effector/AvaEffectorEditorOutlinerContextMenu.cpp
new file mode 100644
--- /dev/null
+++ b/Engine/Plugins/Experimental/Avalanche/Source/AvalancheEffectorsEditor/Private/Effector/AvaEffectorEditorOutlinerContextMenu.cpp
@@ -0,0 +1,92 @@
+// Copyright Epic Games, Inc. All Rights Reserved.
+
+#include "Effector/AvaEffectorEditorOutlinerContextMenu.h"
+
+#include "Effector/CEEffectorComponent.h"
+#include "Effector/Menus/CEEditorEffectorMenuContext.h"
+#include "Effector/Menus/CEEditorEffectorMenuOptions.h"
+#include "GameFramework/Actor.h"
+#include "Item/AvaOutlinerActor.h"
+#include "Item/AvaOutlinerComponent.h"
+#include "Subsystems/CEEditorEffectorSubsystem.h"
+#include "ToolMenuContext/AvaOutlinerItemsContext.h"
+
+void FAvaEffectorEditorOutlinerContextMenu::OnExtendOutlinerContextMenu(UToolMenu* InToolMenu)
+{
+	if (!InToolMenu)
+	{
+		return;
+	}
+
+	const UAvaOutlinerItemsContext* OutlinerItemsContext = InToolMenu->Context.FindContext<UAvaOutlinerItemsContext>();
+
+	if (!IsValid(OutlinerItemsContext))
+	{
+		return;
+	}
+
+	UCEEditorEffectorSubsystem* EffectorSubsystem = UCEEditorEffectorSubsystem::Get();
+
+	if (!IsValid(EffectorSubsystem))
+	{
+		return;
+	}
+
+	TSet<UObject*> ContextObjects;
+	GetContextObjects(OutlinerItemsContext, ContextObjects);
+
+	if (ContextObjects.IsEmpty())
+	{
+		return;
+	}
+
+	const FCEEditorEffectorMenuContext Context(ContextObjects);
+	FCEEditorEffectorMenuOptions Options(
+		{
+			ECEEditorEffectorMenuType::Enable
+			, ECEEditorEffectorMenuType::Disable
+		}
+	);
+	Options.UseTransact(true);
+	Options.CreateSubMenu(true);
+
+	EffectorSubsystem->FillEffectorMenu(InToolMenu, Context, Options);
+}
+
+void FAvaEffectorEditorOutlinerContextMenu::GetContextObjects(const UAvaOutlinerItemsContext* InContext, TSet<UObject*>& OutObjects)
+{
+	OutObjects.Empty();
+
+	if (!IsValid(InContext) || InContext->GetItems().IsEmpty())
+	{
+		return;
+	}
+
+	for (const FAvaOutlinerItemWeakPtr& ItemWeak : InContext->GetItems())
+	{
+		const TSharedPtr<IAvaOutlinerItem> Item = ItemWeak.Pin();
+
+		// is it an actor
+		if (const FAvaOutlinerActor* ActorItem = Item->CastTo<FAvaOutlinerActor>())
+		{
+			if (AActor* Actor = ActorItem->GetActor())
+			{
+				if (Actor->FindComponentByClass<UCEEffectorComponent>())
+				{
+					OutObjects.Add(Actor);
+				}
+			}
+		}
+		// is it a component
+		else if (const FAvaOutlinerComponent* ComponentItem = Item->CastTo<FAvaOutlinerComponent>())
+		{
+			if (USceneComponent* Component = ComponentItem->GetComponent())
+			{
+				if (Component->IsA<UCEEffectorComponent>())
+				{
+					OutObjects.Add(Component);
+				}
+			}
+		}
+	}
+}
diff --git a/Engine/Plugins/Experimental/Avalanche/Source/AvalancheEffectorsEditor/Private/Effector/AvaEffectorEditorOutlinerContextMenu.h b/Engine/Plugins/Experimental/Avalanche/Source/AvalancheEffectorsEditor/Private/Effector/AvaEffectorEditorOutlinerContextMenu.h
new file mode 100644
--- /dev/null
+++ b/Engine/Plugins/Experimental/Avalanche/Source/AvalancheEffectorsEditor/Private/Effector/AvaEffectorEditorOutlinerContextMenu.h
@@ -0,0 +1,21 @@
+// Copyright Epic Games, Inc. All Rights Reserved.
+
+#pragma once
+
+#include "Containers/Set.h"
+#include "ToolMenus.h"
+#include "UObject/Object.h"
+#include "UObject/WeakObjectPtrTemplates.h"
+
+class UAvaOutlinerItemsContext;
+
+/** Context menu actions in outliner related to effector */
+class FAvaEffectorEditorOutlinerContextMenu
+{
+public:
+	/** Extend outliner context menu */
+	static void OnExtendOutlinerContextMenu(UToolMenu* InToolMenu);
+
+	/** Gets all pertinent context object in this context */
+	static void GetContextObjects(const UAvaOutlinerItemsContext* InContext, TSet<UObject*>& OutObjects);
+};
\ No newline at end of file
diff --git a/Engine/Plugins/Experimental/ClonerEffector/Source/ClonerEffector/Private/Cloner/CEClonerComponent.cpp b/Engine/Plugins/Experimental/ClonerEffector/Source/ClonerEffector/Private/Cloner/CEClonerComponent.cpp
--- a/Engine/Plugins/Experimental/ClonerEffector/Source/ClonerEffector/Private/Cloner/CEClonerComponent.cpp
+++ b/Engine/Plugins/Experimental/ClonerEffector/Source/ClonerEffector/Private/Cloner/CEClonerComponent.cpp
@@ -60,6 +60,8 @@ UCEClonerComponent::UCEClonerComponent()
 
 	if (!IsTemplate())
 	{
+		UCEClonerSubsystem::OnClonerSetEnabled().AddUObject(this, &UCEClonerComponent::OnClonerSetEnabled);
+
 		// Bind to delegate to detect material changes
 #if WITH_EDITOR
 		FCoreUObjectDelegates::OnObjectPropertyChanged.RemoveAll(this);
diff --git a/Engine/Plugins/Experimental/ClonerEffector/Source/ClonerEffector/Private/Effector/CEEffectorComponent.cpp b/Engine/Plugins/Experimental/ClonerEffector/Source/ClonerEffector/Private/Effector/CEEffectorComponent.cpp
--- a/Engine/Plugins/Experimental/ClonerEffector/Source/ClonerEffector/Private/Effector/CEEffectorComponent.cpp
+++ b/Engine/Plugins/Experimental/ClonerEffector/Source/ClonerEffector/Private/Effector/CEEffectorComponent.cpp
@@ -38,6 +38,7 @@ UCEEffectorComponent::UCEEffectorComponent()
 		const TArray<FName> ModeNames = GetEffectorModeNames();
 		ModeName = !ModeNames.IsEmpty() ? FName(ModeNames[0]) : NAME_None;
 
+		UCEEffectorSubsystem::OnEffectorSetEnabled().AddUObject(this, &UCEEffectorComponent::OnEffectorSetEnabled);
 		TransformUpdated.AddUObject(this, &UCEEffectorComponent::OnTransformUpdated);
 
 #if WITH_EDITOR
diff --git a/Engine/Plugins/Experimental/ClonerEffector/Source/ClonerEffector/Private/Subsystems/CEClonerSubsystem.cpp b/Engine/Plugins/Experimental/ClonerEffector/Source/ClonerEffector/Private/Subsystems/CEClonerSubsystem.cpp
--- a/Engine/Plugins/Experimental/ClonerEffector/Source/ClonerEffector/Private/Subsystems/CEClonerSubsystem.cpp
+++ b/Engine/Plugins/Experimental/ClonerEffector/Source/ClonerEffector/Private/Subsystems/CEClonerSubsystem.cpp
@@ -3,6 +3,7 @@
 #include "Subsystems/CEClonerSubsystem.h"
 
 #include "Cloner/CEClonerComponent.h"
+#include "Cloner/Extensions/CEClonerEffectorExtension.h"
 #include "Cloner/Extensions/CEClonerExtensionBase.h"
 #include "Cloner/Layouts/CEClonerCircleLayout.h"
 #include "Cloner/Layouts/CEClonerCylinderLayout.h"
@@ -18,7 +19,14 @@
 #include "UObject/Class.h"
 #include "UObject/UObjectIterator.h"
 
+#if WITH_EDITOR
+#include "ScopedTransaction.h"
+#endif
+
 UCEClonerSubsystem::FOnSubsystemInitialized UCEClonerSubsystem::OnSubsystemInitializedDelegate;
+UCEClonerSubsystem::FOnClonerSetEnabled UCEClonerSubsystem::OnClonerSetEnabledDelegate;
+
+#define LOCTEXT_NAMESPACE "CEEffectorSubsystem"
 
 UCEClonerSubsystem* UCEClonerSubsystem::Get()
 {
@@ -240,6 +248,120 @@ UCEClonerExtensionBase* UCEClonerSubsystem::CreateNewExtension(FName InExtension
 	return NewObject<UCEClonerExtensionBase>(InCloner, ExtensionClass->Get(), NAME_None, RF_Transactional);
 }
 
+void UCEClonerSubsystem::SetClonersEnabled(const TSet<UCEClonerComponent*>& InCloners, bool bInEnable, bool bInShouldTransact)
+{
+	if (InCloners.IsEmpty())
+	{
+		return;
+	}
+
+#if WITH_EDITOR
+	const FText TransactionText = bInEnable
+		? LOCTEXT("SetClonersEnabled", "Cloners enabled")
+		: LOCTEXT("SetClonersDisabled", "Cloners disabled");
+
+	FScopedTransaction Transaction(TransactionText, bInShouldTransact);
+#endif
+
+	for (UCEClonerComponent* Cloner : InCloners)
+	{
+		if (!IsValid(Cloner))
+		{
+			continue;
+		}
+
+#if WITH_EDITOR
+		Cloner->Modify();
+#endif
+
+		Cloner->SetEnabled(bInEnable);
+	}
+}
+
+void UCEClonerSubsystem::SetLevelClonersEnabled(const UWorld* InWorld, bool bInEnable, bool bInShouldTransact)
+{
+	if (!IsValid(InWorld))
+	{
+		return;
+	}
+
+#if WITH_EDITOR
+	const FText TransactionText = bInEnable
+		? LOCTEXT("SetLevelClonersEnabled", "Level cloners enabled")
+		: LOCTEXT("SetLevelClonersDisabled", "Level cloners disabled");
+
+	FScopedTransaction Transaction(TransactionText, bInShouldTransact);
+#endif
+
+	OnClonerSetEnabledDelegate.Broadcast(InWorld, bInEnable, bInShouldTransact);
+}
+
+#if WITH_EDITOR
+void UCEClonerSubsystem::ConvertCloners(const TSet<UCEClonerComponent*>& InCloners, ECEClonerMeshConversion InMeshConversion)
+{
+	if (InCloners.IsEmpty())
+	{
+		return;
+	}
+
+	using namespace UE::ClonerEffector::Conversion;
+
+	for (UCEClonerComponent* ClonerComponent : InCloners)
+	{
+		if (!IsValid(ClonerComponent) || !ClonerComponent->GetEnabled())
+		{
+			continue;
+		}
+
+		switch(InMeshConversion)
+		{
+			case ECEClonerMeshConversion::StaticMesh:
+				ClonerComponent->ConvertToStaticMesh();
+			break;
+
+			case ECEClonerMeshConversion::StaticMeshes:
+				ClonerComponent->ConvertToStaticMeshes();
+			break;
+
+			case ECEClonerMeshConversion::DynamicMesh:
+				ClonerComponent->ConvertToDynamicMesh();
+			break;
+
+			case ECEClonerMeshConversion::DynamicMeshes:
+				ClonerComponent->ConvertToDynamicMeshes();
+			break;
+
+			case ECEClonerMeshConversion::InstancedStaticMesh:
+				ClonerComponent->ConvertToInstancedStaticMeshes();
+			break;
+
+			default:;
+		}
+	}
+}
+
+void UCEClonerSubsystem::CreateLinkedEffector(const TSet<UCEClonerComponent*>& InCloners)
+{
+	if (InCloners.IsEmpty())
+	{
+		return;
+	}
+
+	for (UCEClonerComponent* ClonerComponent : InCloners)
+	{
+		if (!IsValid(ClonerComponent))
+		{
+			continue;
+		}
+
+		if (UCEClonerEffectorExtension* EffectorExtension = ClonerComponent->GetExtension<UCEClonerEffectorExtension>())
+		{
+			EffectorExtension->CreateLinkedEffector();
+		}
+	}
+}
+#endif
+
 TArray<FName> UCEClonerSubsystem::GetLayoutNames() const
 {
 	TArray<FName> LayoutNames;
@@ -295,4 +417,6 @@ void UCEClonerSubsystem::ScanForRegistrableClasses()
 			RegisterExtensionClass(ExtensionClass);
 		}
 	}
-}
\ No newline at end of file
+}
+
+#undef LOCTEXT_NAMESPACE
\ No newline at end of file
diff --git a/Engine/Plugins/Experimental/ClonerEffector/Source/ClonerEffector/Private/Subsystems/CEEffectorSubsystem.cpp b/Engine/Plugins/Experimental/ClonerEffector/Source/ClonerEffector/Private/Subsystems/CEEffectorSubsystem.cpp
--- a/Engine/Plugins/Experimental/ClonerEffector/Source/ClonerEffector/Private/Subsystems/CEEffectorSubsystem.cpp
+++ b/Engine/Plugins/Experimental/ClonerEffector/Source/ClonerEffector/Private/Subsystems/CEEffectorSubsystem.cpp
@@ -19,10 +19,17 @@
 #include "Effector/Types/CEEffectorTorusType.h"
 #include "Effector/Types/CEEffectorUnboundType.h"
 
+#if WITH_EDITOR
+#include "ScopedTransaction.h"
+#endif
+
 DEFINE_LOG_CATEGORY_STATIC(LogACEEffectorSubsystem, Log, All);
 
 UCEEffectorSubsystem::FOnSubsystemInitialized UCEEffectorSubsystem::OnSubsystemInitializedDelegate;
 UCEEffectorSubsystem::FOnEffectorIdentifierChanged UCEEffectorSubsystem::OnEffectorIdentifierChangedDelegate;
+UCEEffectorSubsystem::FOnEffectorSetEnabled UCEEffectorSubsystem::OnEffectorSetEnabledDelegate;
+
+#define LOCTEXT_NAMESPACE "CEEffectorSubsystem"
 
 UCEEffectorSubsystem* UCEEffectorSubsystem::Get()
 {
@@ -236,6 +243,54 @@ UCEEffectorExtensionBase* UCEEffectorSubsystem::CreateNewExtension(FName InExten
 	return NewObject<UCEEffectorExtensionBase>(InEffector, ExtensionClass->Get(), NAME_None, RF_Transactional);
 }
 
+void UCEEffectorSubsystem::SetEffectorsEnabled(const TSet<UCEEffectorComponent*>& InEffectors, bool bInEnable, bool bInShouldTransact)
+{
+	if (InEffectors.IsEmpty())
+	{
+		return;
+	}
+
+#if WITH_EDITOR
+	const FText TransactionText = bInEnable
+		? LOCTEXT("SetEffectorsEnabled", "Effectors enabled")
+		: LOCTEXT("SetEffectorsDisabled", "Effectors disabled");
+
+	FScopedTransaction Transaction(TransactionText, bInShouldTransact);
+#endif
+
+	for (UCEEffectorComponent* Effector : InEffectors)
+	{
+		if (!IsValid(Effector))
+		{
+			continue;
+		}
+
+#if WITH_EDITOR
+		Effector->Modify();
+#endif
+
+		Effector->SetEnabled(bInEnable);
+	}
+}
+
+void UCEEffectorSubsystem::SetLevelEffectorsEnabled(const UWorld* InWorld, bool bInEnable, bool bInShouldTransact)
+{
+	if (!IsValid(InWorld))
+	{
+		return;
+	}
+
+#if WITH_EDITOR
+	const FText TransactionText = bInEnable
+		? LOCTEXT("SetLevelEffectorsEnabled", "Level effectors enabled")
+		: LOCTEXT("SetLevelEffectorsDisabled", "Level effectors disabled");
+
+	FScopedTransaction Transaction(TransactionText, bInShouldTransact);
+#endif
+
+	OnEffectorSetEnabledDelegate.Broadcast(InWorld, bInEnable, bInShouldTransact);
+}
+
 void UCEEffectorSubsystem::UpdateEffectorChannel(const UWorld* InWorld)
 {
 	if (!IsValid(InWorld))
@@ -334,3 +389,5 @@ bool UCEEffectorSubsystem::IsTickable() const
 {
 	return !EffectorsWeak.IsEmpty();
 }
+
+#undef LOCTEXT_NAMESPACE
diff --git a/Engine/Plugins/Experimental/ClonerEffector/Source/ClonerEffector/Private/Subsystems/CEEffectorSubsystem.h b/Engine/Plugins/Experimental/ClonerEffector/Source/ClonerEffector/Private/Subsystems/CEEffectorSubsystem.h
deleted file mode 100644
diff --git a/Engine/Plugins/Experimental/ClonerEffector/Source/ClonerEffector/Public/Subsystems/CEClonerSubsystem.h b/Engine/Plugins/Experimental/ClonerEffector/Source/ClonerEffector/Public/Subsystems/CEClonerSubsystem.h
--- a/Engine/Plugins/Experimental/ClonerEffector/Source/ClonerEffector/Public/Subsystems/CEClonerSubsystem.h
+++ b/Engine/Plugins/Experimental/ClonerEffector/Source/ClonerEffector/Public/Subsystems/CEClonerSubsystem.h
@@ -20,6 +20,11 @@ public:
 		return OnSubsystemInitializedDelegate;
 	}
 
+	static TMulticastDelegateRegistration<void(const UWorld*, bool, bool)>& OnClonerSetEnabled()
+	{
+		return OnClonerSetEnabledDelegate;
+	}
+
 	/** Get this subsystem instance */
 	CLONEREFFECTOR_API static UCEClonerSubsystem* Get();
 
@@ -60,10 +65,27 @@ public:
 	/** Creates a new extension instance for a cloner */
 	UCEClonerExtensionBase* CreateNewExtension(FName InExtensionName, UCEClonerComponent* InCloner);
 
+	/** Set cloners state and optionally transact */
+	CLONEREFFECTOR_API void SetClonersEnabled(const TSet<UCEClonerComponent*>& InCloners, bool bInEnable, bool bInShouldTransact);
+
+	/** Set cloners state in world and optionally transact */
+	CLONEREFFECTOR_API void SetLevelClonersEnabled(const UWorld* InWorld, bool bInEnable, bool bInShouldTransact);
+
+#if WITH_EDITOR
+	/** Converts cloners simulation to a mesh */
+	CLONEREFFECTOR_API void ConvertCloners(const TSet<UCEClonerComponent*>& InCloners, ECEClonerMeshConversion InMeshConversion);
+
+	/** Create cloners linked effector */
+	CLONEREFFECTOR_API void CreateLinkedEffector(const TSet<UCEClonerComponent*>& InCloners);
+#endif
 protected:
 	DECLARE_MULTICAST_DELEGATE(FOnSubsystemInitialized)
 	CLONEREFFECTOR_API static FOnSubsystemInitialized OnSubsystemInitializedDelegate;
 
+	/** Delegate to change state of cloners in a world */
+	DECLARE_MULTICAST_DELEGATE_ThreeParams(FOnClonerSetEnabled, const UWorld* /** InWorld */, bool /** bInEnabled */, bool /** bInTransact */)
+	static FOnClonerSetEnabled OnClonerSetEnabledDelegate;
+
 	//~ Begin USubsystem
 	virtual void Initialize(FSubsystemCollectionBase& Collection) override;
 	//~ End USubsystem
diff --git a/Engine/Plugins/Experimental/ClonerEffector/Source/ClonerEffector/Public/Subsystems/CEEffectorSubsystem.h b/Engine/Plugins/Experimental/ClonerEffector/Source/ClonerEffector/Public/Subsystems/CEEffectorSubsystem.h
new file mode 100644
--- /dev/null
+++ b/Engine/Plugins/Experimental/ClonerEffector/Source/ClonerEffector/Public/Subsystems/CEEffectorSubsystem.h
@@ -0,0 +1,117 @@
+// Copyright Epic Games, Inc. All Rights Reserved.
+
+#pragma once
+
+#include "Engine/World.h"
+#include "Subsystems/EngineSubsystem.h"
+#include "UObject/WeakObjectPtrTemplates.h"
+#include "CEEffectorSubsystem.generated.h"
+
+class UCEEffectorComponent;
+class UCEEffectorExtensionBase;
+class UNiagaraDataChannelAsset;
+
+UCLASS()
+class UCEEffectorSubsystem : public UEngineSubsystem, public FTickableGameObject
+{
+	GENERATED_BODY()
+
+public:
+	static TMulticastDelegateRegistration<void()>& OnSubsystemInitialized()
+	{
+		return OnSubsystemInitializedDelegate;
+	}
+
+	static TMulticastDelegateRegistration<void(UCEEffectorComponent*, int32, int32)>& OnEffectorIdentifierChanged()
+	{
+		return OnEffectorIdentifierChangedDelegate;
+	}
+
+	static TMulticastDelegateRegistration<void(const UWorld*, bool, bool)>& OnEffectorSetEnabled()
+	{
+		return OnEffectorSetEnabledDelegate;
+	}
+
+	/** Get this subsystem instance */
+	CLONEREFFECTOR_API static UCEEffectorSubsystem* Get();
+
+	/** Registers an effector actor to use it within a effector channel */
+	bool RegisterChannelEffector(UCEEffectorComponent* InEffector);
+
+	/** Unregister an effector actor used within a effector channel */
+	bool UnregisterChannelEffector(UCEEffectorComponent* InEffector);
+
+	/** Get the effector using this channel identifier */
+	UCEEffectorComponent* GetEffectorByChannelIdentifier(int32 InIdentifier) const;
+
+	bool RegisterExtensionClass(UClass* InClass);
+
+	bool UnregisterExtensionClass(UClass* InClass);
+
+	bool IsExtensionClassRegistered(UClass* InClass) const;
+
+	template<typename InExtensionClass>
+	TArray<FName> GetExtensionNames() const
+	{
+		return GetExtensionNames(InExtensionClass::StaticClass());
+	}
+
+	TArray<FName> GetExtensionNames(TSubclassOf<UCEEffectorExtensionBase> InExtensionClass) const;
+
+	/** Based on a extension class, find extension name */
+	FName FindExtensionName(TSubclassOf<UCEEffectorExtensionBase> InClass) const;
+
+	/** Creates a new extension instance for an effector */
+	UCEEffectorExtensionBase* CreateNewExtension(FName InExtensionName, UCEEffectorComponent* InEffector);
+
+	/** Set effectors state and optionally transact */
+	CLONEREFFECTOR_API void SetEffectorsEnabled(const TSet<UCEEffectorComponent*>& InEffectors, bool bInEnable, bool bInShouldTransact);
+
+	/** Set effectors state in world and optionally transact */
+	CLONEREFFECTOR_API void SetLevelEffectorsEnabled(const UWorld* InWorld, bool bInEnable, bool bInShouldTransact);
+
+protected:
+	static constexpr TCHAR DataChannelAssetPath[] = TEXT("/Script/Niagara.NiagaraDataChannelAsset'/ClonerEffector/Channels/NDC_Effector.NDC_Effector'");
+
+	/** Broadcasted when this subsystem is initialized */
+	DECLARE_MULTICAST_DELEGATE(FOnSubsystemInitialized)
+	static FOnSubsystemInitialized OnSubsystemInitializedDelegate;
+
+	/** Broadcasted when this effector identifier changed to update linked cloners */
+	DECLARE_MULTICAST_DELEGATE_ThreeParams(FOnEffectorIdentifierChanged, UCEEffectorComponent* /** InEffector */, int32 /** OldIdentifier */, int32 /** NewIdentifier */)
+	static FOnEffectorIdentifierChanged OnEffectorIdentifierChangedDelegate;
+
+	/** Delegate to change state of effectors in a world */
+	DECLARE_MULTICAST_DELEGATE_ThreeParams(FOnEffectorSetEnabled, const UWorld* /** InWorld */, bool /** bInEnabled */, bool /** bInTransact */)
+	static FOnEffectorSetEnabled OnEffectorSetEnabledDelegate;
+
+	//~ Begin USubsystem
+	virtual void Initialize(FSubsystemCollectionBase& Collection) override;
+	//~ End USubsystem
+
+	//~ Begin FTickableGameObject
+	virtual TStatId GetStatId() const override;
+	virtual void Tick(float InDeltaSeconds) override;
+	virtual bool IsTickable() const override;
+	virtual bool IsTickableInEditor() const override { return true; }
+	virtual ETickableTickType GetTickableTickType() const override { return ETickableTickType::Conditional; }
+	//~ End FTickableGameObject
+
+	/** Updates all registered effectors */
+	void UpdateEffectorChannel(const UWorld* InWorld);
+
+	/** Scan classes and registers them */
+	void ScanForRegistrableClasses();
+
+	/** Linking name to the extension class */
+	UPROPERTY()
+	TMap<FName, TSubclassOf<UCEEffectorExtensionBase>> ExtensionClasses;
+
+	/** Ordered effectors included in this channel */
+	UPROPERTY()
+	TArray<TWeakObjectPtr<UCEEffectorComponent>> EffectorsWeak;
+
+	/** This represents the data channel structure for effector */
+	UPROPERTY()
+	TObjectPtr<UNiagaraDataChannelAsset> EffectorDataChannelAsset;
+};
\ No newline at end of file
diff --git a/Engine/Plugins/Experimental/ClonerEffector/Source/ClonerEffectorEditor/ClonerEffectorEditor.Build.cs b/Engine/Plugins/Experimental/ClonerEffector/Source/ClonerEffectorEditor/ClonerEffectorEditor.Build.cs
--- a/Engine/Plugins/Experimental/ClonerEffector/Source/ClonerEffectorEditor/ClonerEffectorEditor.Build.cs
+++ b/Engine/Plugins/Experimental/ClonerEffector/Source/ClonerEffectorEditor/ClonerEffectorEditor.Build.cs
@@ -13,6 +13,7 @@ public class ClonerEffectorEditor : ModuleRules
 	        {
 		        "Core",
 		        "CoreUObject",
+				"EditorSubsystem",
 		        "Engine",
 		        "Slate",
 		        "SlateCore",
@@ -26,6 +27,7 @@ public class ClonerEffectorEditor : ModuleRules
 		        "ClonerEffector",
 		        "InputCore",
 		        "Projects",
+				"ToolMenus"
 	        }
         );
 
diff --git a/Engine/Plugins/Experimental/ClonerEffector/Source/ClonerEffectorEditor/Private/Cloner/Menus/CEEditorClonerMenu.cpp b/Engine/Plugins/Experimental/ClonerEffector/Source/ClonerEffectorEditor/Private/Cloner/Menus/CEEditorClonerMenu.cpp
new file mode 100644
--- /dev/null
+++ b/Engine/Plugins/Experimental/ClonerEffector/Source/ClonerEffectorEditor/Private/Cloner/Menus/CEEditorClonerMenu.cpp
@@ -0,0 +1,239 @@
+// Copyright Epic Games, Inc. All Rights Reserved.
+
+#include "Cloner/Menus/CEEditorClonerMenu.h"
+
+#include "Cloner/CEClonerComponent.h"
+#include "Engine/World.h"
+#include "Subsystems/CEClonerSubsystem.h"
+
+#define LOCTEXT_NAMESPACE "CEEditorClonerMenu"
+
+FToolMenuSection* UE::ClonerEditor::Menu::FindOrAddClonerSection(UToolMenu* InMenu)
+{
+	static const FName ClonerSectionName("ContextClonerActions");
+
+	FToolMenuSection* ClonerSection = InMenu->FindSection(ClonerSectionName);
+
+	if (!ClonerSection)
+	{
+		ClonerSection = &InMenu->AddSection(ClonerSectionName
+			, LOCTEXT("ContexClonerActions", "Cloner Actions")
+			, FToolMenuInsert(NAME_None, EToolMenuInsertType::First));
+	}
+
+	return ClonerSection;
+}
+
+void UE::ClonerEditor::Menu::FillEnableClonerSection(UToolMenu* InMenu, const FCEEditorClonerMenuData& InMenuData)
+{
+	if (!InMenu
+		|| InMenuData.Context.IsEmpty()
+		|| !InMenuData.Context.ContainsAnyDisabledCloner())
+	{
+		return;
+	}
+
+	FToolMenuSection& EnableClonerSection = InMenu->FindOrAddSection(
+		TEXT("EnableCloner")
+		, LOCTEXT("EnableCloner.Label", "Enable cloner")
+	);
+
+	constexpr bool bEnable = true;
+
+	EnableClonerSection.AddMenuEntry(
+		TEXT("EnableClonerComponent")
+		, LOCTEXT("EnableClonerComponent.Label", "Enable cloner")
+		, LOCTEXT("EnableClonerComponent.Tooltip", "Enable selected cloners")
+		, FSlateIcon()
+		, FUIAction(
+			FExecuteAction::CreateLambda(&ExecuteEnableClonerAction, InMenuData, bEnable)
+		)
+	);
+
+	EnableClonerSection.AddMenuEntry(
+		TEXT("EnableClonerLevel")
+		, LOCTEXT("EnableClonerLevel.Label", "Enable level cloner")
+		, LOCTEXT("EnableClonerLevel.Tooltip", "Enable selected level cloners")
+		, FSlateIcon()
+		, FUIAction(
+			FExecuteAction::CreateLambda(&ExecuteEnableLevelClonerAction, InMenuData, bEnable)
+		)
+	);
+}
+
+void UE::ClonerEditor::Menu::FillDisableClonerSection(UToolMenu* InMenu, const FCEEditorClonerMenuData& InMenuData)
+{
+	if (!InMenu
+		|| InMenuData.Context.IsEmpty()
+		|| !InMenuData.Context.ContainsAnyEnabledCloner())
+	{
+		return;
+	}
+
+	FToolMenuSection& DisableClonerSection = InMenu->FindOrAddSection(
+		TEXT("DisableCloner")
+		, LOCTEXT("DisableCloner.Label", "Disable cloner")
+	);
+
+	constexpr bool bEnable = false;
+
+	DisableClonerSection.AddMenuEntry(
+		TEXT("DisableClonerComponent")
+		, LOCTEXT("DisableClonerComponent.Label", "Disable cloner")
+		, LOCTEXT("DisableClonerComponent.Tooltip", "Disable selected cloners")
+		, FSlateIcon()
+		, FUIAction(
+			FExecuteAction::CreateLambda(&ExecuteEnableClonerAction, InMenuData, bEnable)
+		)
+	);
+
+	DisableClonerSection.AddMenuEntry(
+		TEXT("DisableClonerLevel")
+		, LOCTEXT("DisableClonerLevel.Label", "Disable level cloner")
+		, LOCTEXT("DisableClonerLevel.Tooltip", "Disable selected level cloners")
+		, FSlateIcon()
+		, FUIAction(
+			FExecuteAction::CreateLambda(&ExecuteEnableLevelClonerAction, InMenuData, bEnable)
+		)
+	);
+}
+
+void UE::ClonerEditor::Menu::FillCreateClonerEffectorSection(UToolMenu* InMenu, const FCEEditorClonerMenuData& InMenuData)
+{
+	if (!InMenu
+		|| InMenuData.Context.IsEmpty()
+		|| !InMenuData.Context.ContainsAnyEnabledCloner())
+	{
+		return;
+	}
+
+	FToolMenuSection* ClonerSection = FindOrAddClonerSection(InMenu);
+
+	check(ClonerSection)
+
+	ClonerSection->AddMenuEntry(
+		TEXT("CreateClonerLinkedEffector")
+		, LOCTEXT("CreateClonerLinkedEffector.Label", "Create linked effector")
+		, LOCTEXT("CreateClonerLinkedEffector.Tooltip", "Create linked effector for selected cloners")
+		, FSlateIcon()
+		, FUIAction(
+			FExecuteAction::CreateLambda(&ExecuteCreateClonerEffectorAction, InMenuData)
+		)
+	);
+}
+
+void UE::ClonerEditor::Menu::FillConvertClonerSection(UToolMenu* InMenu, const FCEEditorClonerMenuData& InMenuData)
+{
+	if (!InMenu
+		|| InMenuData.Context.IsEmpty()
+		|| !InMenuData.Context.ContainsAnyEnabledCloner())
+	{
+		return;
+	}
+
+	FToolMenuSection& ConvertClonerSection = InMenu->FindOrAddSection(
+		TEXT("ConvertCloner")
+		, LOCTEXT("ConvertCloner.Label", "Convert cloner")
+	);
+
+	ConvertClonerSection.AddMenuEntry(
+		TEXT("ConvertClonerToStaticMesh")
+		, LOCTEXT("ConvertClonerToStaticMesh.Label", "To Static Mesh")
+		, LOCTEXT("ConvertClonerToStaticMesh.Tooltip", "Convert selected cloners to static mesh")
+		, FSlateIcon()
+		, FUIAction(
+			FExecuteAction::CreateLambda(&ExecuteConvertClonerAction, InMenuData, ECEClonerMeshConversion::StaticMesh)
+		)
+	);
+
+	ConvertClonerSection.AddMenuEntry(
+		TEXT("ConvertClonerToStaticMeshes")
+		, LOCTEXT("ConvertClonerToStaticMeshes.Label", "To Static Meshes")
+		, LOCTEXT("ConvertClonerToStaticMeshes.Tooltip", "Convert selected cloners to static meshes")
+		, FSlateIcon()
+		, FUIAction(
+			FExecuteAction::CreateLambda(&ExecuteConvertClonerAction, InMenuData, ECEClonerMeshConversion::StaticMeshes)
+		)
+	);
+
+	ConvertClonerSection.AddMenuEntry(
+		TEXT("ConvertClonerToDynamicMesh")
+		, LOCTEXT("ConvertClonerToDynamicMesh.Label", "To Dynamic Mesh")
+		, LOCTEXT("ConvertClonerToDynamicMesh.Tooltip", "Convert selected cloners to dynamic mesh")
+		, FSlateIcon()
+		, FUIAction(
+			FExecuteAction::CreateLambda(&ExecuteConvertClonerAction, InMenuData, ECEClonerMeshConversion::DynamicMesh)
+		)
+	);
+
+	ConvertClonerSection.AddMenuEntry(
+		TEXT("ConvertClonerToDynamicMeshes")
+		, LOCTEXT("ConvertClonerToDynamicMeshes.Label", "To Dynamic Meshes")
+		, LOCTEXT("ConvertClonerToDynamicMeshes.Tooltip", "Convert selected cloners to dynamic meshes")
+		, FSlateIcon()
+		, FUIAction(
+			FExecuteAction::CreateLambda(&ExecuteConvertClonerAction, InMenuData, ECEClonerMeshConversion::DynamicMeshes)
+		)
+	);
+
+	ConvertClonerSection.AddMenuEntry(
+		TEXT("ConvertClonerToInstancedStaticMesh")
+		, LOCTEXT("ConvertClonerToInstancedStaticMesh.Label", "To Instanced Static Mesh")
+		, LOCTEXT("ConvertClonerToInstancedStaticMesh.Tooltip", "Convert selected cloners to instanced static mesh")
+		, FSlateIcon()
+		, FUIAction(
+			FExecuteAction::CreateLambda(&ExecuteConvertClonerAction, InMenuData, ECEClonerMeshConversion::InstancedStaticMesh)
+		)
+	);
+}
+
+void UE::ClonerEditor::Menu::ExecuteEnableClonerAction(const FCEEditorClonerMenuData& InMenuData, bool bInEnable)
+{
+	UCEClonerSubsystem* Subsystem = UCEClonerSubsystem::Get();
+
+	if (InMenuData.Context.IsEmpty() || !Subsystem)
+	{
+		return;
+	}
+
+	Subsystem->SetClonersEnabled(InMenuData.Context.GetComponents(), bInEnable, InMenuData.Options.ShouldTransact());
+}
+
+void UE::ClonerEditor::Menu::ExecuteEnableLevelClonerAction(const FCEEditorClonerMenuData& InMenuData, bool bInEnable)
+{
+	UCEClonerSubsystem* Subsystem = UCEClonerSubsystem::Get();
+	const UWorld* World = InMenuData.Context.GetWorld();
+
+	if (!IsValid(World) || !Subsystem)
+	{
+		return;
+	}
+
+	Subsystem->SetLevelClonersEnabled(World, bInEnable, InMenuData.Options.ShouldTransact());
+}
+
+void UE::ClonerEditor::Menu::ExecuteCreateClonerEffectorAction(const FCEEditorClonerMenuData& InMenuData)
+{
+	UCEClonerSubsystem* Subsystem = UCEClonerSubsystem::Get();
+
+	if (!Subsystem)
+	{
+		return;
+	}
+
+	Subsystem->CreateLinkedEffector(InMenuData.Context.GetComponents());
+}
+
+void UE::ClonerEditor::Menu::ExecuteConvertClonerAction(const FCEEditorClonerMenuData& InMenuData, ECEClonerMeshConversion InToMeshType)
+{
+	UCEClonerSubsystem* Subsystem = UCEClonerSubsystem::Get();
+
+	if (InMenuData.Context.IsEmpty() || !Subsystem)
+	{
+		return;
+	}
+
+	Subsystem->ConvertCloners(InMenuData.Context.GetEnabledCloners(), InToMeshType);
+}
+
+#undef LOCTEXT_NAMESPACE
diff --git a/Engine/Plugins/Experimental/ClonerEffector/Source/ClonerEffectorEditor/Private/Cloner/Menus/CEEditorClonerMenu.h b/Engine/Plugins/Experimental/ClonerEffector/Source/ClonerEffectorEditor/Private/Cloner/Menus/CEEditorClonerMenu.h
new file mode 100644
--- /dev/null
+++ b/Engine/Plugins/Experimental/ClonerEffector/Source/ClonerEffectorEditor/Private/Cloner/Menus/CEEditorClonerMenu.h
@@ -0,0 +1,46 @@
+// Copyright Epic Games, Inc. All Rights Reserved.
+
+#pragma once
+
+#include "CEClonerEffectorShared.h"
+#include "Cloner/Menus/CEEditorClonerMenuContext.h"
+#include "Cloner/Menus/CEEditorClonerMenuOptions.h"
+#include "Templates/SharedPointer.h"
+#include "ToolMenus.h"
+
+namespace UE::ClonerEditor::Menu
+{
+	/** Used internally to group menu data together */
+	struct FCEEditorClonerMenuData
+	{
+		FCEEditorClonerMenuData(const FCEEditorClonerMenuContext& InContext, const FCEEditorClonerMenuOptions& InOptions)
+			: Context(InContext)
+			, Options(InOptions)
+		{}
+
+		const FCEEditorClonerMenuContext Context;
+		const FCEEditorClonerMenuOptions Options;
+	};
+
+	/** Sections */
+
+	FToolMenuSection* FindOrAddClonerSection(UToolMenu* InMenu);
+
+	void FillEnableClonerSection(UToolMenu* InMenu, const FCEEditorClonerMenuData& InMenuData);
+
+	void FillDisableClonerSection(UToolMenu* InMenu, const FCEEditorClonerMenuData& InMenuData);
+
+	void FillCreateClonerEffectorSection(UToolMenu* InMenu, const FCEEditorClonerMenuData& InMenuData);
+
+	void FillConvertClonerSection(UToolMenu* InMenu, const FCEEditorClonerMenuData& InMenuData);
+
+	/** Actions */
+
+	void ExecuteEnableClonerAction(const FCEEditorClonerMenuData& InMenuData, bool bInEnable);
+
+	void ExecuteEnableLevelClonerAction(const FCEEditorClonerMenuData& InMenuData, bool bInEnable);
+
+	void ExecuteCreateClonerEffectorAction(const FCEEditorClonerMenuData& InMenuData);
+
+	void ExecuteConvertClonerAction(const FCEEditorClonerMenuData& InMenuData, ECEClonerMeshConversion InToMeshType);
+}
diff --git a/Engine/Plugins/Experimental/ClonerEffector/Source/ClonerEffectorEditor/Private/Cloner/Menus/CEEditorClonerMenuContext.cpp b/Engine/Plugins/Experimental/ClonerEffector/Source/ClonerEffectorEditor/Private/Cloner/Menus/CEEditorClonerMenuContext.cpp
new file mode 100644
--- /dev/null
+++ b/Engine/Plugins/Experimental/ClonerEffector/Source/ClonerEffectorEditor/Private/Cloner/Menus/CEEditorClonerMenuContext.cpp
@@ -0,0 +1,112 @@
+// Copyright Epic Games, Inc. All Rights Reserved.
+
+#include "Cloner/Menus/CEEditorClonerMenuContext.h"
+
+#include "GameFramework/Actor.h"
+#include "Cloner/CEClonerComponent.h"
+
+FCEEditorClonerMenuContext::FCEEditorClonerMenuContext(const TSet<UObject*>& InObjects)
+{
+	for (UObject* Object : InObjects)
+	{
+		if (!IsValid(Object))
+		{
+			continue;
+		}
+
+		if (const AActor* Actor = Cast<AActor>(Object))
+		{
+			TArray<UCEClonerComponent*> ClonerComponents;
+			Actor->GetComponents(ClonerComponents, /** IncludeChildren */false);
+
+			for (UCEClonerComponent* Component : ClonerComponents)
+			{
+				if (IsValid(Component))
+				{
+					ContextComponents.Add(Component);
+				}
+			}
+		}
+		else if (UCEClonerComponent* Component = Cast<UCEClonerComponent>(Object))
+		{
+			ContextComponents.Add(Component);
+		}
+	}
+}
+
+const TSet<UCEClonerComponent*>& FCEEditorClonerMenuContext::GetComponents() const
+{
+	return ContextComponents;
+}
+
+TSet<UCEClonerComponent*> FCEEditorClonerMenuContext::GetDisabledCloners() const
+{
+	return GetStateCloners(/** IsEnabled */false);
+}
+
+TSet<UCEClonerComponent*> FCEEditorClonerMenuContext::GetEnabledCloners() const
+{
+	return GetStateCloners(/** IsEnabled */true);
+}
+
+UWorld* FCEEditorClonerMenuContext::GetWorld() const
+{
+	for (const UCEClonerComponent* Component : ContextComponents)
+	{
+		if (IsValid(Component))
+		{
+			return Component->GetWorld();
+		}
+	}
+
+	return nullptr;
+}
+
+bool FCEEditorClonerMenuContext::IsEmpty() const
+{
+	return ContextComponents.IsEmpty();
+}
+
+bool FCEEditorClonerMenuContext::ContainsAnyComponent() const
+{
+	return !ContextComponents.IsEmpty();
+}
+
+bool FCEEditorClonerMenuContext::ContainsAnyDisabledCloner() const
+{
+	return ContainsClonerState(/** IsEnabled */false);
+}
+
+bool FCEEditorClonerMenuContext::ContainsAnyEnabledCloner() const
+{
+	return ContainsClonerState(/** IsEnabled */true);
+}
+
+bool FCEEditorClonerMenuContext::ContainsClonerState(bool bInState) const
+{
+	for (UCEClonerComponent* Component : ContextComponents)
+	{
+		if (IsValid(Component) && Component->GetEnabled() == bInState)
+		{
+			return true;
+		}
+	}
+
+	return false;
+}
+
+TSet<UCEClonerComponent*> FCEEditorClonerMenuContext::GetStateCloners(bool bInState) const
+{
+	TSet<UCEClonerComponent*> Cloners;
+	Cloners.Reserve(ContextComponents.Num());
+
+	for (UCEClonerComponent* Component : ContextComponents)
+	{
+		if (IsValid(Component) && Component->GetEnabled() == bInState)
+		{
+			Cloners.Add(Component);
+		}
+	}
+
+	return Cloners;
+}
\ No newline at end of file
diff --git a/Engine/Plugins/Experimental/ClonerEffector/Source/ClonerEffectorEditor/Private/Cloner/Menus/CEEditorClonerMenuOptions.cpp b/Engine/Plugins/Experimental/ClonerEffector/Source/ClonerEffectorEditor/Private/Cloner/Menus/CEEditorClonerMenuOptions.cpp
new file mode 100644
--- /dev/null
+++ b/Engine/Plugins/Experimental/ClonerEffector/Source/ClonerEffectorEditor/Private/Cloner/Menus/CEEditorClonerMenuOptions.cpp
@@ -0,0 +1,30 @@
+// Copyright Epic Games, Inc. All Rights Reserved.
+
+#include "Cloner/Menus/CEEditorClonerMenuOptions.h"
+
+#include "Misc/EnumClassFlags.h"
+
+FCEEditorClonerMenuOptions::FCEEditorClonerMenuOptions(const TSet<ECEEditorClonerMenuType>& InMenus)
+{
+	for (const ECEEditorClonerMenuType& Menu : InMenus)
+	{
+		MenuTypes |= static_cast<uint8>(Menu);
+	}
+}
+
+FCEEditorClonerMenuOptions& FCEEditorClonerMenuOptions::CreateSubMenu(bool bInCreateSubMenu)
+{
+	bCreateSubMenu = bInCreateSubMenu;
+	return *this;
+}
+
+FCEEditorClonerMenuOptions& FCEEditorClonerMenuOptions::UseTransact(bool bInUseTransact)
+{
+	bUseTransact = bInUseTransact;
+	return *this;
+}
+
+bool FCEEditorClonerMenuOptions::IsMenuType(ECEEditorClonerMenuType InMenuType) const
+{
+	return EnumHasAnyFlags(static_cast<ECEEditorClonerMenuType>(MenuTypes), InMenuType);
+}
\ No newline at end of file
diff --git a/Engine/Plugins/Experimental/ClonerEffector/Source/ClonerEffectorEditor/Private/Effector/Menus/CEEditorEffectorMenu.cpp b/Engine/Plugins/Experimental/ClonerEffector/Source/ClonerEffectorEditor/Private/Effector/Menus/CEEditorEffectorMenu.cpp
new file mode 100644
--- /dev/null
+++ b/Engine/Plugins/Experimental/ClonerEffector/Source/ClonerEffectorEditor/Private/Effector/Menus/CEEditorEffectorMenu.cpp
@@ -0,0 +1,124 @@
+// Copyright Epic Games, Inc. All Rights Reserved.
+
+#include "Effector/Menus/CEEditorEffectorMenu.h"
+
+#include "Subsystems/CEEffectorSubsystem.h"
+
+#define LOCTEXT_NAMESPACE "CEEditorEffectorMenu"
+
+FToolMenuSection* UE::EffectorEditor::Menu::FindOrAddEffectorSection(UToolMenu* InMenu)
+{
+	static const FName EffectorSectionName("ContextEffectorActions");
+
+	FToolMenuSection* EffectorSection = InMenu->FindSection(EffectorSectionName);
+
+	if (!EffectorSection)
+	{
+		EffectorSection = &InMenu->AddSection(EffectorSectionName
+			, LOCTEXT("ContexEffectorActions", "Effector Actions")
+			, FToolMenuInsert(NAME_None, EToolMenuInsertType::First));
+	}
+
+	return EffectorSection;
+}
+
+void UE::EffectorEditor::Menu::FillEnableEffectorSection(UToolMenu* InMenu, const FCEEditorEffectorMenuData& InMenuData)
+{
+	if (!InMenu
+		|| InMenuData.Context.IsEmpty()
+		|| !InMenuData.Context.ContainsAnyDisabledEffectors())
+	{
+		return;
+	}
+
+	FToolMenuSection& EnableEffectorSection = InMenu->FindOrAddSection(
+		TEXT("EnableEffector")
+		, LOCTEXT("EnableEffector.Label", "Enable effector")
+	);
+
+	constexpr bool bEnable = true;
+
+	EnableEffectorSection.AddMenuEntry(
+		TEXT("EnableEffectorComponent")
+		, LOCTEXT("EnableEffectorComponent.Label", "Enable effector")
+		, LOCTEXT("EnableEffectorComponent.Tooltip", "Enable selected effectors")
+		, FSlateIcon()
+		, FUIAction(
+			FExecuteAction::CreateLambda(&ExecuteEnableEffectorAction, InMenuData, bEnable)
+		)
+	);
+
+	EnableEffectorSection.AddMenuEntry(
+		TEXT("EnableEffectorLevel")
+		, LOCTEXT("EnableEffectorLevel.Label", "Enable level effector")
+		, LOCTEXT("EnableEffectorLevel.Tooltip", "Enable selected level effectors")
+		, FSlateIcon()
+		, FUIAction(
+			FExecuteAction::CreateLambda(&ExecuteEnableLevelEffectorAction, InMenuData, bEnable)
+		)
+	);
+}
+
+void UE::EffectorEditor::Menu::FillDisableEffectorSection(UToolMenu* InMenu, const FCEEditorEffectorMenuData& InMenuData)
+{
+	if (!InMenu
+		|| InMenuData.Context.IsEmpty()
+		|| !InMenuData.Context.ContainsAnyEnabledEffectors())
+	{
+		return;
+	}
+
+	FToolMenuSection& DisableEffectorSection = InMenu->FindOrAddSection(
+		TEXT("DisableEffector")
+		, LOCTEXT("DisableEffector.Label", "Disable effector")
+	);
+
+	constexpr bool bEnable = false;
+
+	DisableEffectorSection.AddMenuEntry(
+		TEXT("DisableEffectorComponent")
+		, LOCTEXT("DisableEffectorComponent.Label", "Disable effector")
+		, LOCTEXT("DisableEffectorComponent.Tooltip", "Disable selected effectors")
+		, FSlateIcon()
+		, FUIAction(
+			FExecuteAction::CreateLambda(&ExecuteEnableEffectorAction, InMenuData, bEnable)
+		)
+	);
+
+	DisableEffectorSection.AddMenuEntry(
+		TEXT("DisableEffectorLevel")
+		, LOCTEXT("DisableEffectorLevel.Label", "Disable level effector")
+		, LOCTEXT("DisableEffectorLevel.Tooltip", "Disable selected level effectors")
+		, FSlateIcon()
+		, FUIAction(
+			FExecuteAction::CreateLambda(&ExecuteEnableLevelEffectorAction, InMenuData, bEnable)
+		)
+	);
+}
+
+void UE::EffectorEditor::Menu::ExecuteEnableEffectorAction(const FCEEditorEffectorMenuData& InMenuData, bool bInEnable)
+{
+	UCEEffectorSubsystem* Subsystem = UCEEffectorSubsystem::Get();
+
+	if (InMenuData.Context.IsEmpty() || !Subsystem)
+	{
+		return;
+	}
+
+	Subsystem->SetEffectorsEnabled(InMenuData.Context.GetComponents(), bInEnable, InMenuData.Options.ShouldTransact());
+}
+
+void UE::EffectorEditor::Menu::ExecuteEnableLevelEffectorAction(const FCEEditorEffectorMenuData& InMenuData, bool bInEnable)
+{
+	UCEEffectorSubsystem* Subsystem = UCEEffectorSubsystem::Get();
+	const UWorld* World = InMenuData.Context.GetWorld();
+
+	if (!IsValid(World) || !Subsystem)
+	{
+		return;
+	}
+
+	Subsystem->SetLevelEffectorsEnabled(World, bInEnable, InMenuData.Options.ShouldTransact());
+}
+
+#undef LOCTEXT_NAMESPACE
diff --git a/Engine/Plugins/Experimental/ClonerEffector/Source/ClonerEffectorEditor/Private/Effector/Menus/CEEditorEffectorMenu.h b/Engine/Plugins/Experimental/ClonerEffector/Source/ClonerEffectorEditor/Private/Effector/Menus/CEEditorEffectorMenu.h
new file mode 100644
--- /dev/null
+++ b/Engine/Plugins/Experimental/ClonerEffector/Source/ClonerEffectorEditor/Private/Effector/Menus/CEEditorEffectorMenu.h
@@ -0,0 +1,37 @@
+// Copyright Epic Games, Inc. All Rights Reserved.
+
+#pragma once
+
+#include "Effector/Menus/CEEditorEffectorMenuContext.h"
+#include "Effector/Menus/CEEditorEffectorMenuOptions.h"
+#include "Templates/SharedPointer.h"
+#include "ToolMenus.h"
+
+namespace UE::EffectorEditor::Menu
+{
+	/** Used internally to group menu data together */
+	struct FCEEditorEffectorMenuData
+	{
+		FCEEditorEffectorMenuData(const FCEEditorEffectorMenuContext& InContext, const FCEEditorEffectorMenuOptions& InOptions)
+			: Context(InContext)
+			, Options(InOptions)
+		{}
+
+		const FCEEditorEffectorMenuContext Context;
+		const FCEEditorEffectorMenuOptions Options;
+	};
+
+	/** Sections */
+
+	FToolMenuSection* FindOrAddEffectorSection(UToolMenu* InMenu);
+
+	void FillEnableEffectorSection(UToolMenu* InMenu, const FCEEditorEffectorMenuData& InMenuData);
+
+	void FillDisableEffectorSection(UToolMenu* InMenu, const FCEEditorEffectorMenuData& InMenuData);
+
+	/** Actions */
+
+	void ExecuteEnableEffectorAction(const FCEEditorEffectorMenuData& InMenuData, bool bInEnable);
+
+	void ExecuteEnableLevelEffectorAction(const FCEEditorEffectorMenuData& InMenuData, bool bInEnable);
+}
diff --git a/Engine/Plugins/Experimental/ClonerEffector/Source/ClonerEffectorEditor/Private/Effector/Menus/CEEditorEffectorMenuContext.cpp b/Engine/Plugins/Experimental/ClonerEffector/Source/ClonerEffectorEditor/Private/Effector/Menus/CEEditorEffectorMenuContext.cpp
new file mode 100644
--- /dev/null
+++ b/Engine/Plugins/Experimental/ClonerEffector/Source/ClonerEffectorEditor/Private/Effector/Menus/CEEditorEffectorMenuContext.cpp
@@ -0,0 +1,112 @@
+// Copyright Epic Games, Inc. All Rights Reserved.
+
+#include "Effector/Menus/CEEditorEffectorMenuContext.h"
+
+#include "Effector/CEEffectorComponent.h"
+#include "GameFramework/Actor.h"
+
+FCEEditorEffectorMenuContext::FCEEditorEffectorMenuContext(const TSet<UObject*>& InObjects)
+{
+	for (UObject* Object : InObjects)
+	{
+		if (!IsValid(Object))
+		{
+			continue;
+		}
+
+		if (const AActor* Actor = Cast<AActor>(Object))
+		{
+			TArray<UCEEffectorComponent*> EffectorComponents;
+			Actor->GetComponents(EffectorComponents, /** IncludeChildren */false);
+
+			for (UCEEffectorComponent* Component : EffectorComponents)
+			{
+				if (IsValid(Component))
+				{
+					ContextComponents.Add(Component);
+				}
+			}
+		}
+		else if (UCEEffectorComponent* Component = Cast<UCEEffectorComponent>(Object))
+		{
+			ContextComponents.Add(Component);
+		}
+	}
+}
+
+const TSet<UCEEffectorComponent*>& FCEEditorEffectorMenuContext::GetComponents() const
+{
+	return ContextComponents;
+}
+
+TSet<UCEEffectorComponent*> FCEEditorEffectorMenuContext::GetDisabledEffectors() const
+{
+	return GetStateEffectors(/** IsEnabled */ false);
+}
+
+TSet<UCEEffectorComponent*> FCEEditorEffectorMenuContext::GetEnabledEffectors() const
+{
+	return GetStateEffectors(/** IsEnabled */ true);
+}
+
+UWorld* FCEEditorEffectorMenuContext::GetWorld() const
+{
+	for (const UCEEffectorComponent* Component : ContextComponents)
+	{
+		if (IsValid(Component))
+		{
+			return Component->GetWorld();
+		}
+	}
+
+	return nullptr;
+}
+
+bool FCEEditorEffectorMenuContext::IsEmpty() const
+{
+	return ContextComponents.IsEmpty();
+}
+
+bool FCEEditorEffectorMenuContext::ContainsAnyComponent() const
+{
+	return !ContextComponents.IsEmpty();
+}
+
+bool FCEEditorEffectorMenuContext::ContainsAnyDisabledEffectors() const
+{
+	return ContainsEffectorState(/** IsEnabled */ false);
+}
+
+bool FCEEditorEffectorMenuContext::ContainsAnyEnabledEffectors() const
+{
+	return ContainsEffectorState(/** IsEnabled */ true);
+}
+
+bool FCEEditorEffectorMenuContext::ContainsEffectorState(bool bInState) const
+{
+	for (const UCEEffectorComponent* Component : ContextComponents)
+	{
+		if (IsValid(Component) && Component->GetEnabled() == bInState)
+		{
+			return true;
+		}
+	}
+
+	return false;
+}
+
+TSet<UCEEffectorComponent*> FCEEditorEffectorMenuContext::GetStateEffectors(bool bInState) const
+{
+	TSet<UCEEffectorComponent*> Effectors;
+	Effectors.Reserve(ContextComponents.Num());
+
+	for (UCEEffectorComponent* Component : ContextComponents)
+	{
+		if (IsValid(Component) && Component->GetEnabled() == bInState)
+		{
+			Effectors.Add(Component);
+		}
+	}
+
+	return Effectors;
+}
\ No newline at end of file
diff --git a/Engine/Plugins/Experimental/ClonerEffector/Source/ClonerEffectorEditor/Private/Effector/Menus/CEEditorEffectorMenuOptions.cpp b/Engine/Plugins/Experimental/ClonerEffector/Source/ClonerEffectorEditor/Private/Effector/Menus/CEEditorEffectorMenuOptions.cpp
new file mode 100644
--- /dev/null
+++ b/Engine/Plugins/Experimental/ClonerEffector/Source/ClonerEffectorEditor/Private/Effector/Menus/CEEditorEffectorMenuOptions.cpp
@@ -0,0 +1,30 @@
+// Copyright Epic Games, Inc. All Rights Reserved.
+
+#include "Effector/Menus/CEEditorEffectorMenuOptions.h"
+
+#include "Misc/EnumClassFlags.h"
+
+FCEEditorEffectorMenuOptions::FCEEditorEffectorMenuOptions(const TSet<ECEEditorEffectorMenuType>& InMenus)
+{
+	for (const ECEEditorEffectorMenuType& Menu : InMenus)
+	{
+		MenuTypes |= static_cast<uint8>(Menu);
+	}
+}
+
+FCEEditorEffectorMenuOptions& FCEEditorEffectorMenuOptions::CreateSubMenu(bool bInCreateSubMenu)
+{
+	bCreateSubMenu = bInCreateSubMenu;
+	return *this;
+}
+
+FCEEditorEffectorMenuOptions& FCEEditorEffectorMenuOptions::UseTransact(bool bInUseTransact)
+{
+	bUseTransact = bInUseTransact;
+	return *this;
+}
+
+bool FCEEditorEffectorMenuOptions::IsMenuType(ECEEditorEffectorMenuType InMenuType) const
+{
+	return EnumHasAnyFlags(static_cast<ECEEditorEffectorMenuType>(MenuTypes), InMenuType);
+}
\ No newline at end of file
diff --git a/Engine/Plugins/Experimental/ClonerEffector/Source/ClonerEffectorEditor/Private/Subsystems/CEEditorClonerSubsystem.cpp b/Engine/Plugins/Experimental/ClonerEffector/Source/ClonerEffectorEditor/Private/Subsystems/CEEditorClonerSubsystem.cpp
new file mode 100644
--- /dev/null
+++ b/Engine/Plugins/Experimental/ClonerEffector/Source/ClonerEffectorEditor/Private/Subsystems/CEEditorClonerSubsystem.cpp
@@ -0,0 +1,102 @@
+// Copyright Epic Games, Inc. All Rights Reserved.
+
+#include "Subsystems/CEEditorClonerSubsystem.h"
+
+#include "Cloner/Menus/CEEditorClonerMenu.h"
+#include "Editor.h"
+
+#define LOCTEXT_NAMESPACE "CEEditorClonerSubsystem"
+
+UCEEditorClonerSubsystem::UCEEditorClonerSubsystem()
+	: UEditorSubsystem()
+{
+}
+
+UCEEditorClonerSubsystem* UCEEditorClonerSubsystem::Get()
+{
+	if (GEditor)
+	{
+		return GEditor->GetEditorSubsystem<UCEEditorClonerSubsystem>();
+	}
+
+	return nullptr;
+}
+
+void UCEEditorClonerSubsystem::FillClonerMenu(UToolMenu* InMenu, const FCEEditorClonerMenuContext& InContext, const FCEEditorClonerMenuOptions& InOptions)
+{
+	if (!IsValid(InMenu) || InContext.IsEmpty())
+	{
+		return;
+	}
+
+	using namespace UE::ClonerEditor::Menu;
+
+	FToolMenuSection* ClonerSection = nullptr;
+	if (InOptions.ShouldCreateSubMenu())
+	{
+		ClonerSection = FindOrAddClonerSection(InMenu);
+	}
+
+	FCEEditorClonerMenuData MenuData(InContext, InOptions);
+
+	if (InOptions.IsMenuType(ECEEditorClonerMenuType::Enable) && InContext.ContainsAnyDisabledCloner())
+	{
+		if (InOptions.ShouldCreateSubMenu())
+		{
+			check(ClonerSection)
+
+			ClonerSection->AddSubMenu(
+				TEXT("EnableClonerMenu"),
+				LOCTEXT("EnableClonerMenu.Label", "Enable cloners"),
+				LOCTEXT("EnableClonerMenu.Tooltip", "Enable selected cloners"),
+				FNewToolMenuDelegate::CreateLambda(&FillEnableClonerSection, MenuData));
+		}
+		else
+		{
+			FillEnableClonerSection(InMenu, MenuData);
+		}
+	}
+
+	if (InOptions.IsMenuType(ECEEditorClonerMenuType::Disable) && InContext.ContainsAnyEnabledCloner())
+	{
+		if (InOptions.ShouldCreateSubMenu())
+		{
+			check(ClonerSection)
+
+			ClonerSection->AddSubMenu(
+				TEXT("DisableClonerMenu"),
+				LOCTEXT("DisableClonerMenu.Label", "Disable cloners"),
+				LOCTEXT("DisableClonerMenu.Tooltip", "Disable selected cloners"),
+				FNewToolMenuDelegate::CreateLambda(&FillDisableClonerSection, MenuData));
+		}
+		else
+		{
+			FillDisableClonerSection(InMenu, MenuData);
+		}
+	}
+
+	if (InOptions.IsMenuType(ECEEditorClonerMenuType::CreateEffector) && InContext.ContainsAnyComponent())
+	{
+		FillCreateClonerEffectorSection(InMenu, MenuData);
+	}
+
+	if (InOptions.IsMenuType(ECEEditorClonerMenuType::Convert) && InContext.ContainsAnyEnabledCloner())
+	{
+		if (InOptions.ShouldCreateSubMenu())
+		{
+			check(ClonerSection)
+
+			ClonerSection->AddSubMenu(
+				TEXT("ConvertClonerMenu"),
+				LOCTEXT("ConvertClonerMenu.Label", "Convert cloners"),
+				LOCTEXT("ConvertClonerMenu.Tooltip", "Convert selected cloners"),
+				FNewToolMenuDelegate::CreateLambda(&FillConvertClonerSection, MenuData));
+		}
+		else
+		{
+			FillConvertClonerSection(InMenu, MenuData);
+		}
+	}
+}
+
+#undef LOCTEXT_NAMESPACE
diff --git a/Engine/Plugins/Experimental/ClonerEffector/Source/ClonerEffectorEditor/Private/Subsystems/CEEditorEffectorSubsystem.cpp b/Engine/Plugins/Experimental/ClonerEffector/Source/ClonerEffectorEditor/Private/Subsystems/CEEditorEffectorSubsystem.cpp
new file mode 100644
--- /dev/null
+++ b/Engine/Plugins/Experimental/ClonerEffector/Source/ClonerEffectorEditor/Private/Subsystems/CEEditorEffectorSubsystem.cpp
@@ -0,0 +1,79 @@
+// Copyright Epic Games, Inc. All Rights Reserved.
+
+#include "Subsystems/CEEditorEffectorSubsystem.h"
+
+#include "Editor.h"
+#include "Effector/Menus/CEEditorEffectorMenu.h"
+
+#define LOCTEXT_NAMESPACE "CEEditorEffectorSubsystem"
+
+UCEEditorEffectorSubsystem::UCEEditorEffectorSubsystem()
+	: UEditorSubsystem()
+{
+}
+
+UCEEditorEffectorSubsystem* UCEEditorEffectorSubsystem::Get()
+{
+	if (GEditor)
+	{
+		return GEditor->GetEditorSubsystem<UCEEditorEffectorSubsystem>();
+	}
+
+	return nullptr;
+}
+
+void UCEEditorEffectorSubsystem::FillEffectorMenu(UToolMenu* InMenu, const FCEEditorEffectorMenuContext& InContext, const FCEEditorEffectorMenuOptions& InOptions)
+{
+	if (!IsValid(InMenu) || InContext.IsEmpty())
+	{
+		return;
+	}
+
+	using namespace UE::EffectorEditor::Menu;
+
+	FToolMenuSection* EffectorSection = nullptr;
+	if (InOptions.ShouldCreateSubMenu())
+	{
+		EffectorSection = FindOrAddEffectorSection(InMenu);
+	}
+
+	FCEEditorEffectorMenuData MenuData(InContext, InOptions);
+
+	if (InOptions.IsMenuType(ECEEditorEffectorMenuType::Enable) && InContext.ContainsAnyDisabledEffectors())
+	{
+		if (InOptions.ShouldCreateSubMenu())
+		{
+			check(EffectorSection)
+
+			EffectorSection->AddSubMenu(
+				TEXT("EnableEffectorMenu"),
+				LOCTEXT("EnableEffectorMenu.Label", "Enable effectors"),
+				LOCTEXT("EnableEffectorMenu.Tooltip", "Enable selected effectors"),
+				FNewToolMenuDelegate::CreateLambda(&FillEnableEffectorSection, MenuData));
+		}
+		else
+		{
+			FillEnableEffectorSection(InMenu, MenuData);
+		}
+	}
+
+	if (InOptions.IsMenuType(ECEEditorEffectorMenuType::Disable) && InContext.ContainsAnyEnabledEffectors())
+	{
+		if (InOptions.ShouldCreateSubMenu())
+		{
+			check(EffectorSection)
+
+			EffectorSection->AddSubMenu(
+				TEXT("DisableEffectorMenu"),
+				LOCTEXT("DisableEffectorMenu.Label", "Disable effectors"),
+				LOCTEXT("DisableEffectorMenu.Tooltip", "Disable selected effectors"),
+				FNewToolMenuDelegate::CreateLambda(&FillDisableEffectorSection, MenuData));
+		}
+		else
+		{
+			FillDisableEffectorSection(InMenu, MenuData);
+		}
+	}
+}
+
+#undef LOCTEXT_NAMESPACE
diff --git a/Engine/Plugins/Experimental/ClonerEffector/Source/ClonerEffectorEditor/Public/Cloner/Menus/CEEditorClonerMenuContext.h b/Engine/Plugins/Experimental/ClonerEffector/Source/ClonerEffectorEditor/Public/Cloner/Menus/CEEditorClonerMenuContext.h
new file mode 100644
--- /dev/null
+++ b/Engine/Plugins/Experimental/ClonerEffector/Source/ClonerEffectorEditor/Public/Cloner/Menus/CEEditorClonerMenuContext.h
@@ -0,0 +1,35 @@
+// Copyright Epic Games, Inc. All Rights Reserved.
+
+#pragma once
+
+#include "Containers/Set.h"
+#include "UObject/Object.h"
+#include "UObject/WeakObjectPtrTemplates.h"
+
+class AActor;
+class UCEClonerComponent;
+class UObject;
+class UWorld;
+
+/** Menu objects to apply actions on */
+struct CLONEREFFECTOREDITOR_API FCEEditorClonerMenuContext
+{
+	FCEEditorClonerMenuContext() {}
+	explicit FCEEditorClonerMenuContext(const TSet<UObject*>& InObjects);
+
+	const TSet<UCEClonerComponent*>& GetComponents() const;
+	TSet<UCEClonerComponent*> GetDisabledCloners() const;
+	TSet<UCEClonerComponent*> GetEnabledCloners() const;
+	UWorld* GetWorld() const;
+
+	bool IsEmpty() const;
+	bool ContainsAnyComponent() const;
+	bool ContainsAnyDisabledCloner() const;
+	bool ContainsAnyEnabledCloner() const;
+
+protected:
+	bool ContainsClonerState(bool bInState) const;
+	TSet<UCEClonerComponent*> GetStateCloners(bool bInState) const;
+
+	TSet<UCEClonerComponent*> ContextComponents;
+};
\ No newline at end of file
diff --git a/Engine/Plugins/Experimental/ClonerEffector/Source/ClonerEffectorEditor/Public/Cloner/Menus/CEEditorClonerMenuEnums.h b/Engine/Plugins/Experimental/ClonerEffector/Source/ClonerEffectorEditor/Public/Cloner/Menus/CEEditorClonerMenuEnums.h
new file mode 100644
--- /dev/null
+++ b/Engine/Plugins/Experimental/ClonerEffector/Source/ClonerEffectorEditor/Public/Cloner/Menus/CEEditorClonerMenuEnums.h
@@ -0,0 +1,14 @@
+// Copyright Epic Games, Inc. All Rights Reserved.
+
+#pragma once
+
+#include "HAL/Platform.h"
+
+/** Menu type available */
+enum class ECEEditorClonerMenuType : uint8
+{
+	Enable = 1 << 0,
+	Disable = 1 << 1,
+	CreateEffector = 1 << 2,
+	Convert = 1 << 3
+};
\ No newline at end of file
diff --git a/Engine/Plugins/Experimental/ClonerEffector/Source/ClonerEffectorEditor/Public/Cloner/Menus/CEEditorClonerMenuOptions.h b/Engine/Plugins/Experimental/ClonerEffector/Source/ClonerEffectorEditor/Public/Cloner/Menus/CEEditorClonerMenuOptions.h
new file mode 100644
--- /dev/null
+++ b/Engine/Plugins/Experimental/ClonerEffector/Source/ClonerEffectorEditor/Public/Cloner/Menus/CEEditorClonerMenuOptions.h
@@ -0,0 +1,41 @@
+// Copyright Epic Games, Inc. All Rights Reserved.
+
+#pragma once
+
+#include "CEEditorClonerMenuEnums.h"
+#include "Containers/Set.h"
+
+/** Menu options to customize */
+struct CLONEREFFECTOREDITOR_API FCEEditorClonerMenuOptions
+{
+	FCEEditorClonerMenuOptions() {}
+	explicit FCEEditorClonerMenuOptions(const TSet<ECEEditorClonerMenuType>& InMenus);
+	explicit FCEEditorClonerMenuOptions(const uint8 InMenus)
+		: MenuTypes(InMenus)
+	{}
+
+	FCEEditorClonerMenuOptions& CreateSubMenu(bool bInCreateSubMenu);
+	FCEEditorClonerMenuOptions& UseTransact(bool bInUseTransact);
+
+	bool IsMenuType(ECEEditorClonerMenuType InMenuType) const;
+
+	bool ShouldTransact() const
+	{
+		return bUseTransact;
+	}
+
+	bool ShouldCreateSubMenu() const
+	{
+		return bCreateSubMenu;
+	}
+
+protected:
+	/** What type of menu should be generated */
+	uint8 MenuTypes = 0;
+
+	/** Create a transaction for actions performed using the menu */
+	bool bUseTransact = true;
+
+	/** Creates the section inside a submenu */
+	bool bCreateSubMenu = false;
+};
\ No newline at end of file
diff --git a/Engine/Plugins/Experimental/ClonerEffector/Source/ClonerEffectorEditor/Public/Effector/Menus/CEEditorEffectorMenuContext.h b/Engine/Plugins/Experimental/ClonerEffector/Source/ClonerEffectorEditor/Public/Effector/Menus/CEEditorEffectorMenuContext.h
new file mode 100644
--- /dev/null
+++ b/Engine/Plugins/Experimental/ClonerEffector/Source/ClonerEffectorEditor/Public/Effector/Menus/CEEditorEffectorMenuContext.h
@@ -0,0 +1,34 @@
+// Copyright Epic Games, Inc. All Rights Reserved.
+
+#pragma once
+
+#include "Containers/Set.h"
+#include "UObject/Object.h"
+#include "UObject/WeakObjectPtrTemplates.h"
+
+class UCEEffectorComponent;
+class UObject;
+class UWorld;
+
+/** Menu objects to apply actions on */
+struct CLONEREFFECTOREDITOR_API FCEEditorEffectorMenuContext
+{
+	FCEEditorEffectorMenuContext() {}
+	explicit FCEEditorEffectorMenuContext(const TSet<UObject*>& InObjects);
+
+	const TSet<UCEEffectorComponent*>& GetComponents() const;
+	TSet<UCEEffectorComponent*> GetDisabledEffectors() const;
+	TSet<UCEEffectorComponent*> GetEnabledEffectors() const;
+	UWorld* GetWorld() const;
+
+	bool IsEmpty() const;
+	bool ContainsAnyComponent() const;
+	bool ContainsAnyDisabledEffectors() const;
+	bool ContainsAnyEnabledEffectors() const;
+
+protected:
+	bool ContainsEffectorState(bool bInState) const;
+	TSet<UCEEffectorComponent*> GetStateEffectors(bool bInState) const;
+
+	TSet<UCEEffectorComponent*> ContextComponents;
+};
\ No newline at end of file
diff --git a/Engine/Plugins/Experimental/ClonerEffector/Source/ClonerEffectorEditor/Public/Effector/Menus/CEEditorEffectorMenuEnums.h b/Engine/Plugins/Experimental/ClonerEffector/Source/ClonerEffectorEditor/Public/Effector/Menus/CEEditorEffectorMenuEnums.h
new file mode 100644
--- /dev/null
+++ b/Engine/Plugins/Experimental/ClonerEffector/Source/ClonerEffectorEditor/Public/Effector/Menus/CEEditorEffectorMenuEnums.h
@@ -0,0 +1,12 @@
+// Copyright Epic Games, Inc. All Rights Reserved.
+
+#pragma once
+
+#include "HAL/Platform.h"
+
+/** Menu type available */
+enum class ECEEditorEffectorMenuType : uint8
+{
+	Enable = 1 << 0,
+	Disable = 1 << 1
+};
\ No newline at end of file
diff --git a/Engine/Plugins/Experimental/ClonerEffector/Source/ClonerEffectorEditor/Public/Effector/Menus/CEEditorEffectorMenuOptions.h b/Engine/Plugins/Experimental/ClonerEffector/Source/ClonerEffectorEditor/Public/Effector/Menus/CEEditorEffectorMenuOptions.h
new file mode 100644
--- /dev/null
+++ b/Engine/Plugins/Experimental/ClonerEffector/Source/ClonerEffectorEditor/Public/Effector/Menus/CEEditorEffectorMenuOptions.h
@@ -0,0 +1,41 @@
+// Copyright Epic Games, Inc. All Rights Reserved.
+
+#pragma once
+
+#include "CEEditorEffectorMenuEnums.h"
+#include "Containers/Set.h"
+
+/** Menu options to customize */
+struct CLONEREFFECTOREDITOR_API FCEEditorEffectorMenuOptions
+{
+	FCEEditorEffectorMenuOptions() {}
+	explicit FCEEditorEffectorMenuOptions(const TSet<ECEEditorEffectorMenuType>& InMenus);
+	explicit FCEEditorEffectorMenuOptions(const uint8 InMenus)
+		: MenuTypes(InMenus)
+	{}
+
+	FCEEditorEffectorMenuOptions& CreateSubMenu(bool bInCreateSubMenu);
+	FCEEditorEffectorMenuOptions& UseTransact(bool bInUseTransact);
+
+	bool IsMenuType(ECEEditorEffectorMenuType InMenuType) const;
+
+	bool ShouldTransact() const
+	{
+		return bUseTransact;
+	}
+
+	bool ShouldCreateSubMenu() const
+	{
+		return bCreateSubMenu;
+	}
+
+protected:
+	/** What type of menu should be generated */
+	uint8 MenuTypes = 0;
+
+	/** Create a transaction for actions performed using the menu */
+	bool bUseTransact = true;
+
+	/** Creates the section inside a submenu */
+	bool bCreateSubMenu = false;
+};
\ No newline at end of file
diff --git a/Engine/Plugins/Experimental/ClonerEffector/Source/ClonerEffectorEditor/Public/Subsystems/CEEditorClonerSubsystem.h b/Engine/Plugins/Experimental/ClonerEffector/Source/ClonerEffectorEditor/Public/Subsystems/CEEditorClonerSubsystem.h
new file mode 100644
--- /dev/null
+++ b/Engine/Plugins/Experimental/ClonerEffector/Source/ClonerEffectorEditor/Public/Subsystems/CEEditorClonerSubsystem.h
@@ -0,0 +1,28 @@
+// Copyright Epic Games, Inc. All Rights Reserved.
+
+#pragma once
+
+#include "EditorSubsystem.h"
+#include "Templates/SharedPointer.h"
+#include "UObject/WeakObjectPtr.h"
+#include "CEEditorClonerSubsystem.generated.h"
+
+struct FCEEditorClonerMenuContext;
+struct FCEEditorClonerMenuOptions;
+class UToolMenu;
+
+/** Singleton class that handles editor operations for cloners */
+UCLASS(MinimalAPI)
+class UCEEditorClonerSubsystem : public UEditorSubsystem
+{
+	GENERATED_BODY()
+
+public:
+	UCEEditorClonerSubsystem();
+
+	/** Get this subsystem instance */
+	CLONEREFFECTOREDITOR_API static UCEEditorClonerSubsystem* Get();
+
+	/** Fills a menu based on context objects and menu options */
+	CLONEREFFECTOREDITOR_API void FillClonerMenu(UToolMenu* InMenu, const FCEEditorClonerMenuContext& InContext, const FCEEditorClonerMenuOptions& InOptions);
+};
\ No newline at end of file
diff --git a/Engine/Plugins/Experimental/ClonerEffector/Source/ClonerEffectorEditor/Public/Subsystems/CEEditorEffectorSubsystem.h b/Engine/Plugins/Experimental/ClonerEffector/Source/ClonerEffectorEditor/Public/Subsystems/CEEditorEffectorSubsystem.h
new file mode 100644
--- /dev/null
+++ b/Engine/Plugins/Experimental/ClonerEffector/Source/ClonerEffectorEditor/Public/Subsystems/CEEditorEffectorSubsystem.h
@@ -0,0 +1,28 @@
+// Copyright Epic Games, Inc. All Rights Reserved.
+
+#pragma once
+
+#include "EditorSubsystem.h"
+#include "Templates/SharedPointer.h"
+#include "UObject/WeakObjectPtr.h"
+#include "CEEditorEffectorSubsystem.generated.h"
+
+struct FCEEditorEffectorMenuContext;
+struct FCEEditorEffectorMenuOptions;
+class UToolMenu;
+
+/** Singleton class that handles editor operations for effectors */
+UCLASS(MinimalAPI)
+class UCEEditorEffectorSubsystem : public UEditorSubsystem
+{
+	GENERATED_BODY()
+
+public:
+	UCEEditorEffectorSubsystem();
+
+	/** Get this subsystem instance */
+	CLONEREFFECTOREDITOR_API static UCEEditorEffectorSubsystem* Get();
+
+	/** Fills a menu based on context objects and menu options */
+	CLONEREFFECTOREDITOR_API void FillEffectorMenu(UToolMenu* InMenu, const FCEEditorEffectorMenuContext& InContext, const FCEEditorEffectorMenuOptions& InOptions);
+};
\ No newline at end of file
