MotionDesign : ClonerEffector
- Fix cloner and effector details panel customization crash due to property invalid read result

diff --git a/Engine/Plugins/Experimental/ClonerEffector/Source/ClonerEffectorEditor/Private/Cloner/Customizations/CEEditorClonerComponentDetailCustomization.cpp b/Engine/Plugins/Experimental/ClonerEffector/Source/ClonerEffectorEditor/Private/Cloner/Customizations/CEEditorClonerComponentDetailCustomization.cpp
--- a/Engine/Plugins/Experimental/ClonerEffector/Source/ClonerEffectorEditor/Private/Cloner/Customizations/CEEditorClonerComponentDetailCustomization.cpp
+++ b/Engine/Plugins/Experimental/ClonerEffector/Source/ClonerEffectorEditor/Private/Cloner/Customizations/CEEditorClonerComponentDetailCustomization.cpp
@@ -67,7 +67,12 @@ void FCEEditorClonerComponentDetailCustomization::CustomizeDetails(IDetailLayout
 		TSharedPtr<IPropertyHandle> ActiveExtensionProperty = ActiveExtensionArrayProperty->GetElement(ElementIndex);
 
 		UObject* ObjectValue;
-		ActiveExtensionProperty->GetValue(ObjectValue);
+		FPropertyAccess::Result ReadResult = ActiveExtensionProperty->GetValue(ObjectValue);
+
+		if (ReadResult != FPropertyAccess::Success)
+		{
+			continue;
+		}
 
 		UCEClonerExtensionBase* ActiveExtension = Cast<UCEClonerExtensionBase>(ObjectValue);
 
diff --git a/Engine/Plugins/Experimental/ClonerEffector/Source/ClonerEffectorEditor/Private/Effector/Customizations/CEEditorEffectorComponentDetailCustomization.cpp b/Engine/Plugins/Experimental/ClonerEffector/Source/ClonerEffectorEditor/Private/Effector/Customizations/CEEditorEffectorComponentDetailCustomization.cpp
--- a/Engine/Plugins/Experimental/ClonerEffector/Source/ClonerEffectorEditor/Private/Effector/Customizations/CEEditorEffectorComponentDetailCustomization.cpp
+++ b/Engine/Plugins/Experimental/ClonerEffector/Source/ClonerEffectorEditor/Private/Effector/Customizations/CEEditorEffectorComponentDetailCustomization.cpp
@@ -102,7 +102,12 @@ void FCEEditorEffectorComponentDetailCustomization::CustomizeDetails(IDetailLayo
 		TSharedPtr<IPropertyHandle> ActiveExtensionProperty = ActiveExtensionArrayProperty->GetElement(ElementIndex);
 
 		UObject* ObjectValue;
-		ActiveExtensionProperty->GetValue(ObjectValue);
+		FPropertyAccess::Result ReadResult = ActiveExtensionProperty->GetValue(ObjectValue);
+
+		if (ReadResult != FPropertyAccess::Success)
+		{
+			continue;
+		}
 
 		UCEEffectorExtensionBase* ActiveExtension = Cast<UCEEffectorExtensionBase>(ObjectValue);
 
