MotionDesign : ClonerEffector
- Fix incorrect transform result when converting cloner to mesh

diff --git a/Engine/Plugins/Experimental/ClonerEffector/Source/ClonerEffector/Private/CEClonerEffectorShared.cpp b/Engine/Plugins/Experimental/ClonerEffector/Source/ClonerEffector/Private/CEClonerEffectorShared.cpp
--- a/Engine/Plugins/Experimental/ClonerEffector/Source/ClonerEffector/Private/CEClonerEffectorShared.cpp
+++ b/Engine/Plugins/Experimental/ClonerEffector/Source/ClonerEffector/Private/CEClonerEffectorShared.cpp
@@ -205,8 +205,8 @@ AStaticMeshActor* UE::ClonerEffector::Conversion::ConvertClonerToStaticMesh(UCEC
 	}
 
 	FCEClonerMeshBuilder ClonerMeshBuilder;
-
-	if (!ClonerMeshBuilder.AppendComponent(InCloner, InCloner->GetComponentTransform()))
+	const FTransform ClonerTransform = InCloner->GetComponentTransform();
+	if (!ClonerMeshBuilder.AppendComponent(InCloner, ClonerTransform))
 	{
 		return NewActor;
 	}
@@ -239,10 +239,7 @@ AStaticMeshActor* UE::ClonerEffector::Conversion::ConvertClonerToStaticMesh(UCEC
 	SpawnParameters.bTemporaryEditorActor = false;
 #endif
 
-	const FVector ClonerLocation = InCloner->GetComponentLocation();
-	const FRotator ClonerRotation = InCloner->GetComponentRotation();
-
-	NewActor = World->SpawnActor<AStaticMeshActor>(ClonerLocation, ClonerRotation, SpawnParameters);
+	NewActor = World->SpawnActor<AStaticMeshActor>(AStaticMeshActor::StaticClass(), ClonerTransform, SpawnParameters);
 
 	if (!NewActor)
 	{
@@ -324,8 +321,8 @@ ADynamicMeshActor* UE::ClonerEffector::Conversion::ConvertClonerToDynamicMesh(UC
 	}
 
 	FCEClonerMeshBuilder ClonerMeshBuilder;
-
-	if (!ClonerMeshBuilder.AppendComponent(InCloner, InCloner->GetComponentTransform()))
+	const FTransform ClonerTransform = InCloner->GetComponentTransform();
+	if (!ClonerMeshBuilder.AppendComponent(InCloner, ClonerTransform))
 	{
 		return NewActor;
 	}
@@ -343,10 +340,7 @@ ADynamicMeshActor* UE::ClonerEffector::Conversion::ConvertClonerToDynamicMesh(UC
 	SpawnParameters.bTemporaryEditorActor = false;
 #endif
 
-	const FVector ClonerLocation = InCloner->GetComponentLocation();
-	const FRotator ClonerRotation = InCloner->GetComponentRotation();
-
-	NewActor = World->SpawnActor<ADynamicMeshActor>(ClonerLocation, ClonerRotation, SpawnParameters);
+	NewActor = World->SpawnActor<ADynamicMeshActor>(ADynamicMeshActor::StaticClass(), ClonerTransform, SpawnParameters);
 
 	if (!NewActor)
 	{
@@ -395,8 +389,8 @@ TArray<AStaticMeshActor*> UE::ClonerEffector::Conversion::ConvertClonerToStaticM
 	}
 
 	FCEClonerMeshBuilder ClonerMeshBuilder;
-
-	if (!ClonerMeshBuilder.AppendComponent(InCloner, InCloner->GetComponentTransform()))
+	const FTransform ClonerTransform = InCloner->GetComponentTransform();
+	if (!ClonerMeshBuilder.AppendComponent(InCloner, ClonerTransform))
 	{
 		return NewActors;
 	}
@@ -429,18 +423,15 @@ TArray<AStaticMeshActor*> UE::ClonerEffector::Conversion::ConvertClonerToStaticM
 	SpawnParameters.bTemporaryEditorActor = false;
 #endif
 
-	const FVector ClonerLocation = InCloner->GetComponentLocation();
-	const FRotator ClonerRotation = InCloner->GetComponentRotation();
-
 	// Create a Group Actor to hold all actors related to this operation
-	AActor* GroupActor = World->SpawnActor<AActor>(ClonerLocation, ClonerRotation, SpawnParameters);
+	AActor* GroupActor = World->SpawnActor<AActor>(SpawnParameters);
 
 	if (!GroupActor)
 	{
 		return NewActors;
 	}
 
-	CreateComponent(GroupActor, USceneComponent::StaticClass());
+	CreateRootComponent(GroupActor, USceneComponent::StaticClass(), ClonerTransform);
 
 #if WITH_EDITOR
 	FActorLabelUtilities::SetActorLabelUnique(GroupActor, Owner->GetActorNameOrLabel() + TEXT("_SM_Instances"));
@@ -479,7 +470,7 @@ TArray<AStaticMeshActor*> UE::ClonerEffector::Conversion::ConvertClonerToStaticM
 
 		for (const FCEClonerMeshBuilder::FCEClonerMeshInstanceData& Instance : Instances)
 		{
-			if (AStaticMeshActor* StaticMeshActor = World->SpawnActor<AStaticMeshActor>(ClonerLocation, ClonerRotation, SpawnParameters))
+			if (AStaticMeshActor* StaticMeshActor = World->SpawnActor<AStaticMeshActor>(AStaticMeshActor::StaticClass(), ClonerTransform, SpawnParameters))
 			{
 				StaticMeshActor->SetMobility(EComponentMobility::Movable);
 
@@ -493,7 +484,7 @@ TArray<AStaticMeshActor*> UE::ClonerEffector::Conversion::ConvertClonerToStaticM
 					StaticMeshComponent->SetMaterial(MaterialIndex, Instance.MeshMaterials[MaterialIndex].Get());
 				}
 
-				StaticMeshActor->AttachToActor(GroupActor, FAttachmentTransformRules::KeepWorldTransform);
+				StaticMeshActor->AttachToActor(GroupActor, FAttachmentTransformRules::KeepRelativeTransform);
 
 				FActorLabelUtilities::SetActorLabelUnique(StaticMeshActor, Owner->GetActorNameOrLabel() + TEXT("_SM_Instance"));
 
@@ -506,7 +497,7 @@ TArray<AStaticMeshActor*> UE::ClonerEffector::Conversion::ConvertClonerToStaticM
 
 	for (int32 Index = 0; Index < ClonerMeshBuilder.GetMeshInstanceCount(); Index++)
 	{
-		if (AStaticMeshActor* StaticMeshActor = World->SpawnActor<AStaticMeshActor>(ClonerLocation, ClonerRotation, SpawnParameters))
+		if (AStaticMeshActor* StaticMeshActor = World->SpawnActor<AStaticMeshActor>(AStaticMeshActor::StaticClass(), ClonerTransform, SpawnParameters))
 		{
 			StaticMeshActor->SetMobility(EComponentMobility::Movable);
 
@@ -525,7 +516,7 @@ TArray<AStaticMeshActor*> UE::ClonerEffector::Conversion::ConvertClonerToStaticM
 				StaticMeshComponent->SetMaterial(MaterialIndex, MeshData.MeshMaterials[MaterialIndex].Get());
 			}
 
-			StaticMeshActor->AttachToActor(GroupActor, FAttachmentTransformRules::KeepWorldTransform);
+			StaticMeshActor->AttachToActor(GroupActor, FAttachmentTransformRules::KeepRelativeTransform);
 
 			NewActors.Add(StaticMeshActor);
 		}
@@ -558,8 +549,8 @@ TArray<ADynamicMeshActor*> UE::ClonerEffector::Conversion::ConvertClonerToDynami
 	}
 
 	FCEClonerMeshBuilder ClonerMeshBuilder;
-
-	if (!ClonerMeshBuilder.AppendComponent(InCloner, InCloner->GetComponentTransform()))
+	const FTransform ClonerTransform = InCloner->GetComponentTransform();
+	if (!ClonerMeshBuilder.AppendComponent(InCloner, ClonerTransform))
 	{
 		return NewActors;
 	}
@@ -577,18 +568,15 @@ TArray<ADynamicMeshActor*> UE::ClonerEffector::Conversion::ConvertClonerToDynami
 	SpawnParameters.bTemporaryEditorActor = false;
 #endif
 
-	const FVector ClonerLocation = InCloner->GetComponentLocation();
-	const FRotator ClonerRotation = InCloner->GetComponentRotation();
-
 	// Create a Group Actor to hold all actors related to this operation
-	AActor* GroupActor = World->SpawnActor<AActor>(ClonerLocation, ClonerRotation, SpawnParameters);
+	AActor* GroupActor = World->SpawnActor<AActor>(SpawnParameters);
 
 	if (!GroupActor)
 	{
 		return NewActors;
 	}
 
-	CreateComponent(GroupActor, USceneComponent::StaticClass());
+	CreateRootComponent(GroupActor, USceneComponent::StaticClass(), ClonerTransform);
 
 #if WITH_EDITOR
 	FActorLabelUtilities::SetActorLabelUnique(GroupActor, Owner->GetActorNameOrLabel() + TEXT("_DM_Instances"));
@@ -596,7 +584,7 @@ TArray<ADynamicMeshActor*> UE::ClonerEffector::Conversion::ConvertClonerToDynami
 
 	for (int32 Index = 0; Index < ClonerMeshBuilder.GetMeshInstanceCount(); Index++)
 	{
-		if (ADynamicMeshActor* DynamicMeshActor = World->SpawnActor<ADynamicMeshActor>(ClonerLocation, ClonerRotation, SpawnParameters))
+		if (ADynamicMeshActor* DynamicMeshActor = World->SpawnActor<ADynamicMeshActor>(ADynamicMeshActor::StaticClass(), ClonerTransform, SpawnParameters))
 		{
 			UDynamicMeshComponent* DynamicMeshComponent = DynamicMeshActor->GetDynamicMeshComponent();
 
@@ -610,7 +598,7 @@ TArray<ADynamicMeshActor*> UE::ClonerEffector::Conversion::ConvertClonerToDynami
 				DynamicMeshComponent->SetMaterial(MaterialIndex, MeshData.MeshMaterials[MaterialIndex].Get());
 			}
 
-			DynamicMeshActor->AttachToActor(GroupActor, FAttachmentTransformRules::KeepWorldTransform);
+			DynamicMeshActor->AttachToActor(GroupActor, FAttachmentTransformRules::KeepRelativeTransform);
 
 #if WITH_EDITOR
 			FActorLabelUtilities::SetActorLabelUnique(DynamicMeshActor, Owner->GetActorNameOrLabel() + TEXT("_DM_Instance"));
@@ -646,8 +634,8 @@ TArray<AActor*> UE::ClonerEffector::Conversion::ConvertClonerToInstancedStaticMe
 	}
 
 	FCEClonerMeshBuilder ClonerMeshBuilder;
-
-	if (!ClonerMeshBuilder.AppendComponent(InCloner, InCloner->GetComponentTransform()))
+	const FTransform ClonerTransform = InCloner->GetComponentTransform();
+	if (!ClonerMeshBuilder.AppendComponent(InCloner, ClonerTransform))
 	{
 		return NewActors;
 	}
@@ -680,18 +668,15 @@ TArray<AActor*> UE::ClonerEffector::Conversion::ConvertClonerToInstancedStaticMe
 	SpawnParameters.bTemporaryEditorActor = false;
 #endif
 
-	const FVector ClonerLocation = InCloner->GetComponentLocation();
-	const FRotator ClonerRotation = InCloner->GetComponentRotation();
-
 	// Create a Group Actor to hold all actors related to this operation
-	AActor* GroupActor = World->SpawnActor<AActor>(ClonerLocation, ClonerRotation, SpawnParameters);
+	AActor* GroupActor = World->SpawnActor<AActor>(SpawnParameters);
 
 	if (!GroupActor)
 	{
 		return NewActors;
 	}
 
-	CreateComponent(GroupActor, USceneComponent::StaticClass());
+	CreateRootComponent(GroupActor, USceneComponent::StaticClass(), ClonerTransform);
 
 #if WITH_EDITOR
 	FActorLabelUtilities::SetActorLabelUnique(GroupActor, Owner->GetActorNameOrLabel() + TEXT("_ISM_Instances"));
@@ -699,9 +684,9 @@ TArray<AActor*> UE::ClonerEffector::Conversion::ConvertClonerToInstancedStaticMe
 
 	for (uint32 MeshIndex : ClonerMeshBuilder.GetMeshIndexes())
 	{
-		if (AActor* ISMActor = World->SpawnActor<AActor>(ClonerLocation, ClonerRotation, SpawnParameters))
+		if (AActor* ISMActor = World->SpawnActor<AActor>(SpawnParameters))
 		{
-			UInstancedStaticMeshComponent* ISMComponent = Cast<UInstancedStaticMeshComponent>(CreateComponent(ISMActor, UInstancedStaticMeshComponent::StaticClass()));
+			UInstancedStaticMeshComponent* ISMComponent = Cast<UInstancedStaticMeshComponent>(CreateRootComponent(ISMActor, UInstancedStaticMeshComponent::StaticClass(), ClonerTransform));
 
 			UStaticMesh* StaticMesh = nullptr;
 
@@ -739,7 +724,7 @@ TArray<AActor*> UE::ClonerEffector::Conversion::ConvertClonerToInstancedStaticMe
 
 			for (const FCEClonerMeshBuilder::FCEClonerMeshInstanceData& Instance : Instances)
 			{
-				ISMComponent->AddInstance(Instance.Transform, /** WorldSpace */false);
+				ISMComponent->AddInstance(Instance.Transform, /** WorldSpace */true);
 			}
 
 			if (!Instances.IsEmpty())
@@ -753,7 +738,7 @@ TArray<AActor*> UE::ClonerEffector::Conversion::ConvertClonerToInstancedStaticMe
 				}
 			}
 
-			ISMActor->AttachToActor(GroupActor, FAttachmentTransformRules::KeepWorldTransform);
+			ISMActor->AttachToActor(GroupActor, FAttachmentTransformRules::KeepRelativeTransform);
 
 #if WITH_EDITOR
 			FActorLabelUtilities::SetActorLabelUnique(ISMActor, Owner->GetActorNameOrLabel() + TEXT("_ISM_Instance"));
@@ -766,7 +751,7 @@ TArray<AActor*> UE::ClonerEffector::Conversion::ConvertClonerToInstancedStaticMe
 	return NewActors;
 }
 
-UActorComponent* UE::ClonerEffector::Conversion::CreateComponent(AActor* InActor, TSubclassOf<USceneComponent> InComponentClass)
+UActorComponent* UE::ClonerEffector::Conversion::CreateRootComponent(AActor* InActor, TSubclassOf<USceneComponent> InComponentClass, const FTransform& InWorldTransform)
 {
 	USceneComponent* const NewComponent = NewObject<USceneComponent>(InActor
 		, InComponentClass.Get()
@@ -782,6 +767,8 @@ UActorComponent* UE::ClonerEffector::Conversion::CreateComponent(AActor* InActor
 	InActor->RerunConstructionScripts();
 #endif
 
+	NewComponent->SetWorldTransform(InWorldTransform);
+
 	return NewComponent;
 }
 
diff --git a/Engine/Plugins/Experimental/ClonerEffector/Source/ClonerEffector/Public/CEClonerEffectorShared.h b/Engine/Plugins/Experimental/ClonerEffector/Source/ClonerEffector/Public/CEClonerEffectorShared.h
--- a/Engine/Plugins/Experimental/ClonerEffector/Source/ClonerEffector/Public/CEClonerEffectorShared.h
+++ b/Engine/Plugins/Experimental/ClonerEffector/Source/ClonerEffector/Public/CEClonerEffectorShared.h
@@ -560,7 +560,8 @@ namespace UE::ClonerEffector
 		/** Convert a cloner to multiple instanced static mesh actors */
 		TArray<AActor*> ConvertClonerToInstancedStaticMeshes(UCEClonerComponent* InCloner);
 
-		UActorComponent* CreateComponent(AActor* InActor, TSubclassOf<USceneComponent> InComponentClass);
+		/** Creates the root component for an actor */
+		UActorComponent* CreateRootComponent(AActor* InActor, TSubclassOf<USceneComponent> InComponentClass, const FTransform& InWorldTransform);
 
 #if WITH_EDITOR
 		/** Pick assets location */
