MotionDesign : ClonerEffector
- Fix sub categories not supported in details panel section customization

diff --git a/Engine/Plugins/Experimental/ClonerEffector/Source/ClonerEffector/Private/Settings/CEClonerEffectorSettings.h b/Engine/Plugins/Experimental/ClonerEffector/Source/ClonerEffector/Private/Settings/CEClonerEffectorSettings.h
--- a/Engine/Plugins/Experimental/ClonerEffector/Source/ClonerEffector/Private/Settings/CEClonerEffectorSettings.h
+++ b/Engine/Plugins/Experimental/ClonerEffector/Source/ClonerEffector/Private/Settings/CEClonerEffectorSettings.h
@@ -79,7 +79,7 @@ protected:
 
 #if WITH_EDITORONLY_DATA
 	/** Reduces the r.TSR.ShadingRejection.Flickering.Period from 3 (default) to 1 if enabled to avoid ghosting artifacts when moving */
-	UPROPERTY(EditInstanceOnly, BlueprintReadWrite, AdvancedDisplay, Category="Cloner", meta=(AllowPrivateAccess="true"))
+	UPROPERTY(Config, EditAnywhere, Category="Cloner")
 	bool bReduceMotionGhosting = false;
 #endif
 
diff --git a/Engine/Plugins/Experimental/ClonerEffector/Source/ClonerEffector/Public/Cloner/Layouts/CEClonerSplineLayout.h b/Engine/Plugins/Experimental/ClonerEffector/Source/ClonerEffector/Public/Cloner/Layouts/CEClonerSplineLayout.h
--- a/Engine/Plugins/Experimental/ClonerEffector/Source/ClonerEffector/Public/Cloner/Layouts/CEClonerSplineLayout.h
+++ b/Engine/Plugins/Experimental/ClonerEffector/Source/ClonerEffector/Public/Cloner/Layouts/CEClonerSplineLayout.h
@@ -62,7 +62,7 @@ public:
 
 #if WITH_EDITOR
 	/** Spawns an actor with a spline component and linked it to this cloner layout */
-	UFUNCTION(CallInEditor, Category="Cloner|Layout")
+	UFUNCTION(CallInEditor, Category="Layout")
 	CLONEREFFECTOR_API void SpawnLinkedSplineActor();
 #endif
 
diff --git a/Engine/Plugins/Experimental/ClonerEffector/Source/ClonerEffectorEditor/Private/Cloner/Customizations/CEEditorClonerComponentDetailCustomization.cpp b/Engine/Plugins/Experimental/ClonerEffector/Source/ClonerEffectorEditor/Private/Cloner/Customizations/CEEditorClonerComponentDetailCustomization.cpp
--- a/Engine/Plugins/Experimental/ClonerEffector/Source/ClonerEffectorEditor/Private/Cloner/Customizations/CEEditorClonerComponentDetailCustomization.cpp
+++ b/Engine/Plugins/Experimental/ClonerEffector/Source/ClonerEffectorEditor/Private/Cloner/Customizations/CEEditorClonerComponentDetailCustomization.cpp
@@ -212,10 +212,17 @@ void FCEEditorClonerComponentDetailCustomization::CustomizeDetails(IDetailLayout
 		{
 			const FName FunctionsCategoryName = FunctionToCategoryPair.Value;
 
-			const TSharedPtr<SVerticalBox> FunctionsWidget = SNew(SVerticalBox);
-
 			IDetailCategoryBuilder& FunctionsCategory = InDetailBuilder.EditCategory(FunctionsCategoryName, FText::FromName(FunctionsCategoryName), ECategoryPriority::Uncommon);
 
+			// Subcategories are not supported in sections
+			if (!FunctionsCategoryName.ToString().Contains(TEXT("|")))
+			{
+				const TSharedRef<FPropertySection> FunctionsSection = PropertyModule.FindOrCreateSection(ComponentClassName, FunctionsCategoryName, FText::FromName(FunctionsCategoryName));
+				FunctionsSection->AddCategory(FunctionsCategoryName);
+			}
+
+			const TSharedPtr<SVerticalBox> FunctionsWidget = SNew(SVerticalBox);
+
 			FunctionsCategory.AddCustomRow(FText::GetEmpty())
 				.WholeRowContent()
 				.HAlign(HAlign_Left)
@@ -223,9 +230,6 @@ void FCEEditorClonerComponentDetailCustomization::CustomizeDetails(IDetailLayout
 					FunctionsWidget.ToSharedRef()
 				];
 
-			const TSharedRef<FPropertySection> FunctionsSection = PropertyModule.FindOrCreateSection(ComponentClassName, FunctionsCategoryName, FText::FromName(FunctionsCategoryName));
-			FunctionsSection->AddCategory(FunctionsCategoryName);
-
 			FunctionToWidget.Add(FunctionToCategoryPair.Key, FunctionsWidget);
 		}
 
