MotionDesign : ClonerEffector
- Adding cloner layout and extension detail customizations for actor picker properties
- Custom node customization to show actor picker with filtered actors based on predicate

diff --git a/Engine/Plugins/Experimental/ClonerEffector/Source/ClonerEffector/Private/Cloner/CEClonerComponent.cpp b/Engine/Plugins/Experimental/ClonerEffector/Source/ClonerEffector/Private/Cloner/CEClonerComponent.cpp
--- a/Engine/Plugins/Experimental/ClonerEffector/Source/ClonerEffector/Private/Cloner/CEClonerComponent.cpp
+++ b/Engine/Plugins/Experimental/ClonerEffector/Source/ClonerEffector/Private/Cloner/CEClonerComponent.cpp
@@ -1003,6 +1003,16 @@ void UCEClonerComponent::ShowMaterialWarning(int32 InMaterialCount)
 		FSlateNotificationManager::Get().AddNotification(NotificationInfo);
 	}
 }
+
+FName UCEClonerComponent::GetActiveExtensionsName()
+{
+	return GET_MEMBER_NAME_CHECKED(UCEClonerComponent, ActiveExtensions);
+}
+
+FName UCEClonerComponent::GetActiveLayoutName()
+{
+	return GET_MEMBER_NAME_CHECKED(UCEClonerComponent, ActiveLayout);
+}
 #endif // WITH_EDITOR
 
 void UCEClonerComponent::InitializeCloner()
diff --git a/Engine/Plugins/Experimental/ClonerEffector/Source/ClonerEffector/Private/Cloner/Extensions/CEClonerEffectorExtension.cpp b/Engine/Plugins/Experimental/ClonerEffector/Source/ClonerEffector/Private/Cloner/Extensions/CEClonerEffectorExtension.cpp
--- a/Engine/Plugins/Experimental/ClonerEffector/Source/ClonerEffector/Private/Cloner/Extensions/CEClonerEffectorExtension.cpp
+++ b/Engine/Plugins/Experimental/ClonerEffector/Source/ClonerEffector/Private/Cloner/Extensions/CEClonerEffectorExtension.cpp
@@ -15,6 +15,13 @@
 
 DEFINE_LOG_CATEGORY_STATIC(LogCEClonerEffectorExtension, Log, All);
 
+#if WITH_EDITOR
+FName UCEClonerEffectorExtension::GetEffectorActorsWeakName()
+{
+	return GET_MEMBER_NAME_CHECKED(UCEClonerEffectorExtension, EffectorActorsWeak);
+}
+#endif
+
 bool UCEClonerEffectorExtension::LinkEffector(AActor* InEffectorActor)
 {
 	if (!IsValid(InEffectorActor)
diff --git a/Engine/Plugins/Experimental/ClonerEffector/Source/ClonerEffector/Private/Cloner/Layouts/CEClonerMeshLayout.cpp b/Engine/Plugins/Experimental/ClonerEffector/Source/ClonerEffector/Private/Cloner/Layouts/CEClonerMeshLayout.cpp
--- a/Engine/Plugins/Experimental/ClonerEffector/Source/ClonerEffector/Private/Cloner/Layouts/CEClonerMeshLayout.cpp
+++ b/Engine/Plugins/Experimental/ClonerEffector/Source/ClonerEffector/Private/Cloner/Layouts/CEClonerMeshLayout.cpp
@@ -9,6 +9,18 @@
 #include "NiagaraDataInterfaceSkeletalMesh.h"
 #include "NiagaraSystem.h"
 
+#if WITH_EDITOR
+FName UCEClonerMeshLayout::GetSampleActorWeakName()
+{
+	return GET_MEMBER_NAME_CHECKED(UCEClonerMeshLayout, SampleActorWeak);
+}
+
+FName UCEClonerMeshLayout::GetAssetName()
+{
+	return GET_MEMBER_NAME_CHECKED(UCEClonerMeshLayout, Asset);
+}
+#endif
+
 void UCEClonerMeshLayout::SetCount(int32 InCount)
 {
 	if (Count == InCount)
diff --git a/Engine/Plugins/Experimental/ClonerEffector/Source/ClonerEffector/Private/Cloner/Layouts/CEClonerSplineLayout.cpp b/Engine/Plugins/Experimental/ClonerEffector/Source/ClonerEffector/Private/Cloner/Layouts/CEClonerSplineLayout.cpp
--- a/Engine/Plugins/Experimental/ClonerEffector/Source/ClonerEffector/Private/Cloner/Layouts/CEClonerSplineLayout.cpp
+++ b/Engine/Plugins/Experimental/ClonerEffector/Source/ClonerEffector/Private/Cloner/Layouts/CEClonerSplineLayout.cpp
@@ -12,6 +12,13 @@
 #include "Editor/EditorEngine.h"
 #endif
 
+#if WITH_EDITOR
+FName UCEClonerSplineLayout::GetSplineActorWeakName()
+{
+	return GET_MEMBER_NAME_CHECKED(UCEClonerSplineLayout, SplineActorWeak);
+}
+#endif
+
 void UCEClonerSplineLayout::SetCount(int32 InCount)
 {
 	if (Count == InCount)
diff --git a/Engine/Plugins/Experimental/ClonerEffector/Source/ClonerEffector/Public/Cloner/CEClonerComponent.h b/Engine/Plugins/Experimental/ClonerEffector/Source/ClonerEffector/Public/Cloner/CEClonerComponent.h
--- a/Engine/Plugins/Experimental/ClonerEffector/Source/ClonerEffector/Public/Cloner/CEClonerComponent.h
+++ b/Engine/Plugins/Experimental/ClonerEffector/Source/ClonerEffector/Public/Cloner/CEClonerComponent.h
@@ -25,7 +25,6 @@ class UCEClonerComponent : public UNiagaraComponent
 	GENERATED_BODY()
 
 	friend class ACEClonerActor;
-	friend class FCEEditorClonerComponentDetailCustomization;
 
 public:
 	/** Only materials transient or part of the content folder can be dirtied, engine or plugins cannot */
@@ -37,6 +36,10 @@ public:
 #if WITH_EDITOR
 	/** Show material warning notification when missing niagara usage flag */
 	static void ShowMaterialWarning(int32 InMaterialCount);
+
+	static CLONEREFFECTOR_API FName GetActiveExtensionsName();
+
+	static CLONEREFFECTOR_API FName GetActiveLayoutName();
 #endif
 
 	TMulticastDelegateRegistration<void(UCEClonerComponent*)>& OnClonerMeshUpdated()
diff --git a/Engine/Plugins/Experimental/ClonerEffector/Source/ClonerEffector/Public/Cloner/Extensions/CEClonerEffectorExtension.h b/Engine/Plugins/Experimental/ClonerEffector/Source/ClonerEffector/Public/Cloner/Extensions/CEClonerEffectorExtension.h
--- a/Engine/Plugins/Experimental/ClonerEffector/Source/ClonerEffector/Public/Cloner/Extensions/CEClonerEffectorExtension.h
+++ b/Engine/Plugins/Experimental/ClonerEffector/Source/ClonerEffector/Public/Cloner/Extensions/CEClonerEffectorExtension.h
@@ -18,6 +18,10 @@ class UCEClonerEffectorExtension : public UCEClonerExtensionBase
 	GENERATED_BODY()
 
 public:
+#if WITH_EDITOR
+	static CLONEREFFECTOR_API FName GetEffectorActorsWeakName();
+#endif
+
 	UCEClonerEffectorExtension()
 		: UCEClonerExtensionBase(
 			TEXT("Effector")
@@ -71,11 +75,11 @@ protected:
 	UPROPERTY(EditInstanceOnly, Category="Effector", meta=(DisplayName="Effectors"))
 	TArray<TWeakObjectPtr<AActor>> EffectorActorsWeak;
 
+private:
 	/** Copy of effectors linked to this cloner to compare diffs */
 	UPROPERTY(Transient, DuplicateTransient, TextExportTransient, NonTransactional)
 	TSet<TWeakObjectPtr<UCEEffectorComponent>> EffectorsInternalWeak;
 
-private:
 #if WITH_EDITOR
 	/** Used for PECP */
 	static const TCEPropertyChangeDispatcher<UCEClonerEffectorExtension> PropertyChangeDispatcher;
diff --git a/Engine/Plugins/Experimental/ClonerEffector/Source/ClonerEffector/Public/Cloner/Layouts/CEClonerMeshLayout.h b/Engine/Plugins/Experimental/ClonerEffector/Source/ClonerEffector/Public/Cloner/Layouts/CEClonerMeshLayout.h
--- a/Engine/Plugins/Experimental/ClonerEffector/Source/ClonerEffector/Public/Cloner/Layouts/CEClonerMeshLayout.h
+++ b/Engine/Plugins/Experimental/ClonerEffector/Source/ClonerEffector/Public/Cloner/Layouts/CEClonerMeshLayout.h
@@ -19,6 +19,11 @@ class UCEClonerMeshLayout : public UCEClonerLayoutBase
 	friend class FAvaClonerActorVisualizer;
 
 public:
+#if WITH_EDITOR
+	static CLONEREFFECTOR_API FName GetSampleActorWeakName();
+	static CLONEREFFECTOR_API FName GetAssetName();
+#endif
+
 	UCEClonerMeshLayout()
 		: UCEClonerLayoutBase(
 			TEXT("Mesh")
@@ -84,22 +89,22 @@ protected:
 
 	void OnSampleMeshTransformed(USceneComponent* InComponent, EUpdateTransformFlags InFlags, ETeleportType InType);
 
-	UPROPERTY(EditInstanceOnly, BlueprintReadWrite, Setter="SetCount", Getter="GetCount", Category="Layout")
-	int32 Count = 3 * 3 * 3;
-
-	UPROPERTY(EditInstanceOnly, BlueprintReadWrite, Setter="SetAsset", Getter="GetAsset", Category="Layout")
+	UPROPERTY(EditInstanceOnly, BlueprintReadWrite, Setter, Getter, Category="Layout")
 	ECEClonerMeshAsset Asset = ECEClonerMeshAsset::StaticMesh;
 
-	UPROPERTY(EditInstanceOnly, BlueprintReadWrite, Setter="SetSampleData", Getter="GetSampleData", Category="Layout")
-	ECEClonerMeshSampleData SampleData = ECEClonerMeshSampleData::Vertices;
-
-	UPROPERTY(EditInstanceOnly, BlueprintReadWrite, Setter="SetSampleActorWeak", Getter="GetSampleActorWeak", DisplayName="SampleActor", Category="Layout")
+	UPROPERTY(EditInstanceOnly, BlueprintReadWrite, Setter, Getter, DisplayName="Sample Actor", Category="Layout")
 	TWeakObjectPtr<AActor> SampleActorWeak;
 
+	UPROPERTY(EditInstanceOnly, BlueprintReadWrite, Setter, Getter, Category="Layout")
+	int32 Count = 3 * 3 * 3;
+
+	UPROPERTY(EditInstanceOnly, BlueprintReadWrite, Setter, Getter, Category="Layout")
+	ECEClonerMeshSampleData SampleData = ECEClonerMeshSampleData::Vertices;
+
+private:
 	UPROPERTY(Transient)
 	TWeakObjectPtr<USceneComponent> SceneComponentWeak;
 
-private:
 #if WITH_EDITOR
 	/** Used for PECP */
 	static const TCEPropertyChangeDispatcher<UCEClonerMeshLayout> PropertyChangeDispatcher;
diff --git a/Engine/Plugins/Experimental/ClonerEffector/Source/ClonerEffector/Public/Cloner/Layouts/CEClonerSplineLayout.h b/Engine/Plugins/Experimental/ClonerEffector/Source/ClonerEffector/Public/Cloner/Layouts/CEClonerSplineLayout.h
--- a/Engine/Plugins/Experimental/ClonerEffector/Source/ClonerEffector/Public/Cloner/Layouts/CEClonerSplineLayout.h
+++ b/Engine/Plugins/Experimental/ClonerEffector/Source/ClonerEffector/Public/Cloner/Layouts/CEClonerSplineLayout.h
@@ -18,6 +18,10 @@ class UCEClonerSplineLayout : public UCEClonerLayoutBase
 	friend class FAvaClonerActorVisualizer;
 
 public:
+#if WITH_EDITOR
+	static CLONEREFFECTOR_API FName GetSplineActorWeakName();
+#endif
+
 	UCEClonerSplineLayout()
 		: UCEClonerLayoutBase(
 			TEXT("Spline")
@@ -90,10 +94,10 @@ protected:
 	UPROPERTY(EditInstanceOnly, BlueprintReadWrite, Setter="SetOrientMesh", Getter="GetOrientMesh", Category="Layout")
     bool bOrientMesh = false;
 
+private:
 	UPROPERTY(Transient)
 	TWeakObjectPtr<USplineComponent> SplineComponentWeak;
 
-private:
 #if WITH_EDITOR
 	/** Used for PECP */
 	static const TCEPropertyChangeDispatcher<UCEClonerSplineLayout> PropertyChangeDispatcher;
diff --git a/Engine/Plugins/Experimental/ClonerEffector/Source/ClonerEffectorEditor/ClonerEffectorEditor.Build.cs b/Engine/Plugins/Experimental/ClonerEffector/Source/ClonerEffectorEditor/ClonerEffectorEditor.Build.cs
--- a/Engine/Plugins/Experimental/ClonerEffector/Source/ClonerEffectorEditor/ClonerEffectorEditor.Build.cs
+++ b/Engine/Plugins/Experimental/ClonerEffector/Source/ClonerEffectorEditor/ClonerEffectorEditor.Build.cs
@@ -27,6 +27,7 @@ public class ClonerEffectorEditor : ModuleRules
 		        "ClonerEffector",
 		        "InputCore",
 		        "Projects",
+				"PropertyEditor",
 				"ToolMenus"
 	        }
         );
diff --git a/Engine/Plugins/Experimental/ClonerEffector/Source/ClonerEffectorEditor/Private/CEEditorModule.cpp b/Engine/Plugins/Experimental/ClonerEffector/Source/ClonerEffectorEditor/Private/CEEditorModule.cpp
--- a/Engine/Plugins/Experimental/ClonerEffector/Source/ClonerEffectorEditor/Private/CEEditorModule.cpp
+++ b/Engine/Plugins/Experimental/ClonerEffector/Source/ClonerEffectorEditor/Private/CEEditorModule.cpp
@@ -4,8 +4,14 @@
 
 #include "Cloner/CEClonerComponent.h"
 #include "Cloner/CEClonerActor.h"
-#include "Cloner/Customizations/CEEditorClonerComponentDetailCustomization.h"
 #include "Cloner/Customizations/CEEditorClonerActorDetailCustomization.h"
+#include "Cloner/Customizations/CEEditorClonerComponentDetailCustomization.h"
+#include "Cloner/Customizations/CEEditorClonerEffectorExtensionDetailCustomization.h"
+#include "Cloner/Customizations/CEEditorClonerMeshLayoutDetailCustomization.h"
+#include "Cloner/Customizations/CEEditorClonerSplineLayoutDetailCustomization.h"
+#include "Cloner/Extensions/CEClonerEffectorExtension.h"
+#include "Cloner/Layouts/CEClonerMeshLayout.h"
+#include "Cloner/Layouts/CEClonerSplineLayout.h"
 #include "Effector/Customizations/CEEditorEffectorComponentDetailCustomization.h"
 #include "Effector/Customizations/CEEditorEffectorTypeDetailCustomization.h"
 #include "Effector/CEEffectorActor.h"
@@ -26,6 +32,9 @@ void FCEEditorModule::StartupModule()
 	// Cloner customization
 	PropertyModule.RegisterCustomClassLayout(ACEClonerActor::StaticClass()->GetFName(), FOnGetDetailCustomizationInstance::CreateStatic(&FCEEditorClonerActorDetailCustomization::MakeInstance));
 	PropertyModule.RegisterCustomClassLayout(UCEClonerComponent::StaticClass()->GetFName(), FOnGetDetailCustomizationInstance::CreateStatic(&FCEEditorClonerComponentDetailCustomization::MakeInstance));
+	PropertyModule.RegisterCustomClassLayout(UCEClonerEffectorExtension::StaticClass()->GetFName(), FOnGetDetailCustomizationInstance::CreateStatic(&FCEEditorClonerEffectorExtensionDetailCustomization::MakeInstance));
+	PropertyModule.RegisterCustomClassLayout(UCEClonerSplineLayout::StaticClass()->GetFName(), FOnGetDetailCustomizationInstance::CreateStatic(&FCEEditorClonerSplineLayoutDetailCustomization::MakeInstance));
+	PropertyModule.RegisterCustomClassLayout(UCEClonerMeshLayout::StaticClass()->GetFName(), FOnGetDetailCustomizationInstance::CreateStatic(&FCEEditorClonerMeshLayoutDetailCustomization::MakeInstance));
 
 	// Effector customization
 	PropertyModule.RegisterCustomClassLayout(ACEEffectorActor::StaticClass()->GetFName(), FOnGetDetailCustomizationInstance::CreateStatic(&FCEEditorEffectorActorDetailCustomization::MakeInstance));
@@ -42,6 +51,9 @@ void FCEEditorModule::ShutdownModule()
 		// Cloner customization
 		PropertyModule.UnregisterCustomClassLayout(ACEClonerActor::StaticClass()->GetFName());
 		PropertyModule.UnregisterCustomClassLayout(UCEClonerComponent::StaticClass()->GetFName());
+		PropertyModule.UnregisterCustomClassLayout(UCEClonerEffectorExtension::StaticClass()->GetFName());
+		PropertyModule.UnregisterCustomClassLayout(UCEClonerSplineLayout::StaticClass()->GetFName());
+		PropertyModule.UnregisterCustomClassLayout(UCEClonerMeshLayout::StaticClass()->GetFName());
 
 		// Effector customization
 		PropertyModule.UnregisterCustomClassLayout(ACEEffectorActor::StaticClass()->GetFName());
diff --git a/Engine/Plugins/Experimental/ClonerEffector/Source/ClonerEffectorEditor/Private/Cloner/Customizations/CEEditorClonerComponentDetailCustomization.cpp b/Engine/Plugins/Experimental/ClonerEffector/Source/ClonerEffectorEditor/Private/Cloner/Customizations/CEEditorClonerComponentDetailCustomization.cpp
--- a/Engine/Plugins/Experimental/ClonerEffector/Source/ClonerEffectorEditor/Private/Cloner/Customizations/CEEditorClonerComponentDetailCustomization.cpp
+++ b/Engine/Plugins/Experimental/ClonerEffector/Source/ClonerEffectorEditor/Private/Cloner/Customizations/CEEditorClonerComponentDetailCustomization.cpp
@@ -2,12 +2,15 @@
 
 #include "Cloner/Customizations/CEEditorClonerComponentDetailCustomization.h"
 
+#include "CEEditorClonerEffectorExtensionDetailCustomization.h"
 #include "Cloner/CEClonerComponent.h"
+#include "Cloner/Extensions/CEClonerEffectorExtension.h"
 #include "Cloner/Extensions/CEClonerExtensionBase.h"
 #include "Cloner/Extensions/CEClonerLifetimeExtension.h"
 #include "Cloner/Layouts/CEClonerLayoutBase.h"
 #include "DetailLayoutBuilder.h"
 #include "DetailWidgetRow.h"
+#include "IDetailChildrenBuilder.h"
 #include "Input/Reply.h"
 #include "Modules/ModuleManager.h"
 #include "NiagaraDataInterfaceCurve.h"
@@ -30,7 +33,7 @@ void FCEEditorClonerComponentDetailCustomization::CustomizeDetails(IDetailLayout
 	InDetailBuilder.HideCategory(TEXT("NiagaraComponent_Utilities"));
 
 	// Remove extension array property
-	TSharedRef<IPropertyHandle> ActiveExtensionsProperty = InDetailBuilder.GetProperty(GET_MEMBER_NAME_CHECKED(UCEClonerComponent, ActiveExtensions), UCEClonerComponent::StaticClass());
+	TSharedRef<IPropertyHandle> ActiveExtensionsProperty = InDetailBuilder.GetProperty(UCEClonerComponent::GetActiveExtensionsName(), UCEClonerComponent::StaticClass());
 
 	if (!ActiveExtensionsProperty->IsValidHandle())
 	{
@@ -46,7 +49,7 @@ void FCEEditorClonerComponentDetailCustomization::CustomizeDetails(IDetailLayout
 	// Everything needs to be below Cloner category
 	int32 StartOrder = InDetailBuilder.EditCategory(TEXT("Cloner")).GetSortOrder() + 1;
 
-	TSharedRef<IPropertyHandle> ActiveLayoutProperty = InDetailBuilder.GetProperty(GET_MEMBER_NAME_CHECKED(UCEClonerComponent, ActiveLayout), UCEClonerComponent::StaticClass());
+	TSharedRef<IPropertyHandle> ActiveLayoutProperty = InDetailBuilder.GetProperty(UCEClonerComponent::GetActiveLayoutName(), UCEClonerComponent::StaticClass());
 
 	if (!ActiveLayoutProperty->IsValidHandle())
 	{
@@ -103,7 +106,20 @@ void FCEEditorClonerComponentDetailCustomization::CustomizeDetails(IDetailLayout
 				}
 				else
 				{
-					ExtensionCategoryBuilder.AddProperty(ChildHandle);
+					const FProperty* Property = ChildHandle->GetProperty();
+
+					/*
+					 * Since AddExternalObject doesn't add hidden properties (EditConditionHides) of the object but only the visible ones
+					 * We add them manually, but this also means the object customization is not used, we need to set it manually from here
+					*/
+					if (Property && Property->GetName() == UCEClonerEffectorExtension::GetEffectorActorsWeakName())
+					{
+						FCEEditorClonerEffectorExtensionDetailCustomization::CustomizeEffectorsProperty(ChildHandle.ToSharedRef(), ExtensionCategoryBuilder);
+					}
+					else
+					{
+						ExtensionCategoryBuilder.AddProperty(ChildHandle);
+					}
 				}
 			}
 		};
@@ -144,6 +160,11 @@ void FCEEditorClonerComponentDetailCustomization::CustomizeDetails(IDetailLayout
 		}
 	}
 
+	if (!InDetailBuilder.GetSelectedObjectsOfType<UCEClonerComponent>().IsEmpty())
+	{
+		return;
+	}
+
 	// Handle ufunctions
 	TMap<FName, FName> FunctionToCategory;
 	for (const TWeakObjectPtr<UCEClonerComponent>& ClonerComponentWeak : ClonerComponentsWeak)
diff --git a/Engine/Plugins/Experimental/ClonerEffector/Source/ClonerEffectorEditor/Private/Cloner/Customizations/CEEditorClonerCustomActorPickerNodeBuilder.cpp b/Engine/Plugins/Experimental/ClonerEffector/Source/ClonerEffectorEditor/Private/Cloner/Customizations/CEEditorClonerCustomActorPickerNodeBuilder.cpp
new file mode 100644
--- /dev/null
+++ b/Engine/Plugins/Experimental/ClonerEffector/Source/ClonerEffectorEditor/Private/Cloner/Customizations/CEEditorClonerCustomActorPickerNodeBuilder.cpp
@@ -0,0 +1,298 @@
+// Copyright Epic Games, Inc. All Rights Reserved.
+
+#include "Cloner/Customizations/CEEditorClonerCustomActorPickerNodeBuilder.h"
+
+#include "DetailWidgetRow.h"
+#include "Editor.h"
+#include "GameFramework/Actor.h"
+#include "PropertyCustomizationHelpers.h"
+#include "ScopedTransaction.h"
+#include "Selection.h"
+#include "Widgets/Input/SComboButton.h"
+#include "Widgets/Text/STextBlock.h"
+
+#define LOCTEXT_NAMESPACE "CEEditorClonerCustomActorPickerNodeBuilder"
+
+FName FCEEditorClonerCustomActorPickerNodeBuilder::GetName() const
+{
+	return TEXT("FCEEditorClonerCustomActorPickerNodeBuilder");
+}
+
+void FCEEditorClonerCustomActorPickerNodeBuilder::OnSelectActor() const
+{
+	if (!GEditor)
+	{
+		return;
+	}
+
+	AActor* Actor = GetPropertyActor();
+
+	if (IsValid(Actor))
+	{
+		const FScopedTransaction Transaction(NSLOCTEXT("UnrealEd", "ClickingOnElements", "Clicking on Elements"));
+
+		constexpr bool bNotifySelectionChanged = true;
+		GEditor->SelectNone(!bNotifySelectionChanged, /** DeselectBSP */true);
+		GEditor->SelectActor(Actor, /** IsSelected */true, bNotifySelectionChanged);
+	}
+}
+
+void FCEEditorClonerCustomActorPickerNodeBuilder::OnUseSelectedActor() const
+{
+	if (AActor* SelectedActor = GetSelectedActor())
+	{
+		if (ActorFilterDelegate.Execute(SelectedActor))
+		{
+			OnActorSelected(SelectedActor);
+		}
+	}
+}
+
+AActor* FCEEditorClonerCustomActorPickerNodeBuilder::GetPropertyActor() const
+{
+	if (!PropertyHandle->IsValidHandle())
+	{
+		return nullptr;
+	}
+
+	UObject* SelectedObject;
+	if (PropertyHandle->GetValue(SelectedObject) != FPropertyAccess::Success)
+	{
+		return nullptr;
+	}
+
+	return Cast<AActor>(SelectedObject);
+}
+
+AActor* FCEEditorClonerCustomActorPickerNodeBuilder::GetSelectedActor() const
+{
+	if (!GEditor || !GEditor->GetSelectedActors())
+	{
+		return nullptr;
+	}
+
+	return GEditor->GetSelectedActors()->GetTop<AActor>();
+}
+
+void FCEEditorClonerCustomActorPickerNodeBuilder::GenerateHeaderRowContent(FDetailWidgetRow& InNodeRow)
+{
+	if (!PropertyHandle->IsValidHandle())
+	{
+		return;
+	}
+
+	bool bIsArrayProperty = false;
+	if (TSharedPtr<IPropertyHandle> ParentHandle = PropertyHandle->GetParentHandle())
+	{
+		if (ParentHandle->AsArray())
+		{
+			bIsArrayProperty = true;
+		}
+	}
+
+	TSharedRef<SHorizontalBox> CustomValueWidget = SNew(SHorizontalBox)
+		+ SHorizontalBox::Slot()
+		.HAlign(HAlign_Fill)
+		.FillWidth(VAlign_Fill)
+		.FillWidth(1.f)
+		[
+			SAssignNew(ComboButton, SComboButton)
+			.ContentPadding(2.0f)
+			.HasDownArrow(true)
+			.OnGetMenuContent_Raw(this, &FCEEditorClonerCustomActorPickerNodeBuilder::GetActorPickerWidget)
+			.ButtonContent()
+			[
+				SNew(STextBlock)
+				.Text_Raw(this, &FCEEditorClonerCustomActorPickerNodeBuilder::GetPickerLabelText)
+			]
+		];
+
+	/* Add use selected actor button */
+	CustomValueWidget->AddSlot()
+		.HAlign(HAlign_Fill)
+		.FillWidth(VAlign_Fill)
+		.AutoWidth()
+		[
+			PropertyCustomizationHelpers::MakeUseSelectedButton(
+				FSimpleDelegate::CreateSP(this, &FCEEditorClonerCustomActorPickerNodeBuilder::OnUseSelectedActor)
+				, TAttribute<FText>(this, &FCEEditorClonerCustomActorPickerNodeBuilder::GetUseSelectTooltipText)
+				, /** IsEnabled */true
+				, /** IsActor */true)
+		];
+
+	/* Add browse to actor button */
+	CustomValueWidget->AddSlot()
+		.HAlign(HAlign_Fill)
+		.FillWidth(VAlign_Fill)
+		.AutoWidth()
+		[
+			PropertyCustomizationHelpers::MakeBrowseButton(
+				FSimpleDelegate::CreateSP(this, &FCEEditorClonerCustomActorPickerNodeBuilder::OnSelectActor)
+				, TAttribute<FText>(this, &FCEEditorClonerCustomActorPickerNodeBuilder::GetSelectTooltipText)
+				, /** IsEnabled */true
+				, /** IsActor */true)
+		];
+
+	/* Add pick level actor button */
+
+	FOnGetAllowedClasses OnGetAllowedClassesDelegate;
+
+	CustomValueWidget->AddSlot()
+		.HAlign(HAlign_Fill)
+		.FillWidth(VAlign_Fill)
+		.AutoWidth()
+		[
+			PropertyCustomizationHelpers::MakeInteractiveActorPicker(
+				OnGetAllowedClassesDelegate
+				, ActorFilterDelegate
+				, FOnActorSelected::CreateSP(this, &FCEEditorClonerCustomActorPickerNodeBuilder::OnActorSelected)
+			)
+		];
+
+	if (bIsArrayProperty)
+	{
+		TSharedRef<SWidget> DefaultValueWidget = PropertyHandle->CreatePropertyValueWidget();
+
+		// Hide original selector widget
+		const TSharedPtr<SWidget> SelectorWidget = FindWidgetType(DefaultValueWidget, TEXT("SPropertyEditorAsset"), /** reverse search */false);
+		check(SelectorWidget)
+		SelectorWidget->SetVisibility(EVisibility::Collapsed);
+
+		// But keep original array widget visible (insert/duplicate/remove)
+		const TSharedPtr<SWidget> ArrayButtons = FindWidgetType(DefaultValueWidget, TEXT("SComboButton"), /** reverse search */true);
+		check(ArrayButtons)
+		ArrayButtons->SetVisibility(EVisibility::Visible);
+
+		CustomValueWidget->AddSlot()
+			.HAlign(HAlign_Fill)
+			.FillWidth(VAlign_Fill)
+			.AutoWidth()
+			[
+				DefaultValueWidget
+			];
+	}
+
+	InNodeRow.NameContent()
+	[
+		PropertyHandle->CreatePropertyNameWidget()
+	];
+
+	InNodeRow.ValueContent()
+	[
+		CustomValueWidget
+	];
+}
+
+void FCEEditorClonerCustomActorPickerNodeBuilder::GenerateChildContent(IDetailChildrenBuilder& InChildrenBuilder)
+{
+}
+
+TSharedPtr<IPropertyHandle> FCEEditorClonerCustomActorPickerNodeBuilder::GetPropertyHandle() const
+{
+	return PropertyHandle;
+}
+
+FText FCEEditorClonerCustomActorPickerNodeBuilder::GetPickerLabelText() const
+{
+	const AActor* Actor = GetPropertyActor();
+
+	if (IsValid(Actor))
+	{
+		return FText::FromString(Actor->GetActorNameOrLabel());
+	}
+
+	return FText::GetEmpty();
+}
+
+FText FCEEditorClonerCustomActorPickerNodeBuilder::GetSelectTooltipText() const
+{
+	const AActor* Actor = GetPropertyActor();
+
+	if (IsValid(Actor))
+	{
+		return FText::Format(LOCTEXT("SelectActor", "Select '{0}' in the viewport"), FText::FromString(Actor->GetName()));
+	}
+
+	return FText::GetEmpty();
+}
+
+FText FCEEditorClonerCustomActorPickerNodeBuilder::GetUseSelectTooltipText() const
+{
+	const AActor* Actor = GetSelectedActor();
+
+	if (IsValid(Actor))
+	{
+		return FText::Format(LOCTEXT("UseSelectedActor", "Use '{0}' from the viewport"), FText::FromString(Actor->GetName()));
+	}
+
+	return FText::GetEmpty();
+}
+
+TSharedRef<SWidget> FCEEditorClonerCustomActorPickerNodeBuilder::GetActorPickerWidget() const
+{
+	FSimpleDelegate EmptyDelegate;
+
+	return PropertyCustomizationHelpers::MakeActorPickerWithMenu(
+		GetPropertyActor()
+		, /** AllowClear */true
+		, ActorFilterDelegate
+		, FOnActorSelected::CreateSP(this, &FCEEditorClonerCustomActorPickerNodeBuilder::OnActorSelected)
+		, EmptyDelegate
+		, EmptyDelegate);
+}
+
+void FCEEditorClonerCustomActorPickerNodeBuilder::OnActorSelected(AActor* InSelection) const
+{
+	if (PropertyHandle->IsValidHandle())
+	{
+		PropertyHandle->SetValue(InSelection);
+	}
+
+	if (ComboButton)
+	{
+		ComboButton->SetIsOpen(false);
+	}
+}
+
+TSharedPtr<SWidget> FCEEditorClonerCustomActorPickerNodeBuilder::FindWidgetType(TSharedPtr<SWidget> InSearchWidget, FName InTypeName, bool bInReverseSearch)
+{
+	if (!InSearchWidget.IsValid())
+	{
+		return nullptr;
+	}
+
+	const FName WidgetTypeName = InSearchWidget->GetType();
+
+	if (WidgetTypeName == InTypeName)
+	{
+		return InSearchWidget;
+	}
+
+	if (FChildren* Children = InSearchWidget->GetAllChildren())
+	{
+		if (bInReverseSearch)
+		{
+			for (int32 Index = Children->Num() - 1; Index >= 0; Index--)
+			{
+				if (TSharedPtr<SWidget> FoundWidget = FindWidgetType(Children->GetChildAt(Index), InTypeName, bInReverseSearch))
+				{
+					return FoundWidget;
+				}
+			}
+		}
+		else
+		{
+			for (int32 Index = 0; Index < Children->Num(); Index++)
+			{
+				if (TSharedPtr<SWidget> FoundWidget = FindWidgetType(Children->GetChildAt(Index), InTypeName, bInReverseSearch))
+				{
+					return FoundWidget;
+				}
+			}
+		}
+	}
+
+	return nullptr;
+}
+
+#undef LOCTEXT_NAMESPACE
\ No newline at end of file
diff --git a/Engine/Plugins/Experimental/ClonerEffector/Source/ClonerEffectorEditor/Private/Cloner/Customizations/CEEditorClonerCustomActorPickerNodeBuilder.h b/Engine/Plugins/Experimental/ClonerEffector/Source/ClonerEffectorEditor/Private/Cloner/Customizations/CEEditorClonerCustomActorPickerNodeBuilder.h
new file mode 100644
--- /dev/null
+++ b/Engine/Plugins/Experimental/ClonerEffector/Source/ClonerEffectorEditor/Private/Cloner/Customizations/CEEditorClonerCustomActorPickerNodeBuilder.h
@@ -0,0 +1,48 @@
+// Copyright Epic Games, Inc. All Rights Reserved.
+
+#pragma once
+
+#include "ActorPickerMode.h"
+#include "IDetailCustomNodeBuilder.h"
+#include "Templates/SharedPointer.h"
+
+class IPropertyHandle;
+class IPropertyUtilities;
+class SComboButton;
+class SWidget;
+
+/** Used to customize an actor picker to filter items based on a delegate */
+class FCEEditorClonerCustomActorPickerNodeBuilder : public IDetailCustomNodeBuilder, public TSharedFromThis<FCEEditorClonerCustomActorPickerNodeBuilder> 
+{
+public:
+	explicit FCEEditorClonerCustomActorPickerNodeBuilder(const TSharedRef<IPropertyHandle>& InPropertyHandle, FOnShouldFilterActor InActorFilterDelegate)
+		: PropertyHandle(InPropertyHandle)
+		, ActorFilterDelegate(InActorFilterDelegate)
+	{}
+
+private:
+	static TSharedPtr<SWidget> FindWidgetType(TSharedPtr<SWidget> InSearchWidget, FName InTypeName, bool bInReverseSearch);
+
+	//~ Begin IDetailCustomNodeBuilder
+	virtual FName GetName() const override;
+	virtual void GenerateHeaderRowContent(FDetailWidgetRow& InNodeRow) override;
+	virtual void GenerateChildContent(IDetailChildrenBuilder& InChildrenBuilder) override;
+	virtual TSharedPtr<IPropertyHandle> GetPropertyHandle() const override;
+	//~ End IDetailCustomNodeBuilder
+
+	FText GetPickerLabelText() const;
+	FText GetSelectTooltipText() const;
+	FText GetUseSelectTooltipText() const;
+	TSharedRef<SWidget> GetActorPickerWidget() const;
+
+	void OnActorSelected(AActor* InSelection) const;
+	void OnSelectActor() const;
+	void OnUseSelectedActor() const;
+
+	AActor* GetPropertyActor() const;
+	AActor* GetSelectedActor() const;
+
+	TSharedRef<IPropertyHandle> PropertyHandle;
+	TSharedPtr<SComboButton> ComboButton;
+	FOnShouldFilterActor ActorFilterDelegate;
+};
diff --git a/Engine/Plugins/Experimental/ClonerEffector/Source/ClonerEffectorEditor/Private/Cloner/Customizations/CEEditorClonerEffectorExtensionDetailCustomization.cpp b/Engine/Plugins/Experimental/ClonerEffector/Source/ClonerEffectorEditor/Private/Cloner/Customizations/CEEditorClonerEffectorExtensionDetailCustomization.cpp
new file mode 100644
--- /dev/null
+++ b/Engine/Plugins/Experimental/ClonerEffector/Source/ClonerEffectorEditor/Private/Cloner/Customizations/CEEditorClonerEffectorExtensionDetailCustomization.cpp
@@ -0,0 +1,53 @@
+// Copyright Epic Games, Inc. All Rights Reserved.
+
+#include "Cloner/Customizations/CEEditorClonerEffectorExtensionDetailCustomization.h"
+
+#include "Cloner/Customizations/CEEditorClonerCustomActorPickerNodeBuilder.h"
+#include "Cloner/Extensions/CEClonerEffectorExtension.h"
+#include "DetailLayoutBuilder.h"
+#include "Effector/CEEffectorComponent.h"
+#include "GameFramework/Actor.h"
+#include "IDetailChildrenBuilder.h"
+#include "PropertyCustomizationHelpers.h"
+#include "PropertyHandle.h"
+
+void FCEEditorClonerEffectorExtensionDetailCustomization::CustomizeDetails(IDetailLayoutBuilder& InDetailBuilder)
+{
+	TSharedRef<IPropertyHandle> EffectorsPropertyHandle = InDetailBuilder.GetProperty(UCEClonerEffectorExtension::GetEffectorActorsWeakName(), UCEClonerEffectorExtension::StaticClass());
+
+	if (!EffectorsPropertyHandle->IsValidHandle())
+	{
+		return;
+	}
+
+	EffectorsPropertyHandle->MarkHiddenByCustomization();
+
+	IDetailCategoryBuilder& EffectorsCategoryBuilder = InDetailBuilder.EditCategory(EffectorsPropertyHandle->GetDefaultCategoryName(), EffectorsPropertyHandle->GetDefaultCategoryText());
+
+	CustomizeEffectorsProperty(EffectorsPropertyHandle, EffectorsCategoryBuilder);
+}
+
+void FCEEditorClonerEffectorExtensionDetailCustomization::CustomizeEffectorsProperty(TSharedRef<IPropertyHandle> InProperty, IDetailCategoryBuilder& InCategoryBuilder)
+{
+	const TSharedRef<FDetailArrayBuilder> EffectorsArrayBuilder = MakeShared<FDetailArrayBuilder>(InProperty, /*GenerateHeader*/true, /*DisplayResetToDefault*/true, /*bDisplayElementNum*/true);
+
+	auto FilterActorDelegate = FOnShouldFilterActor::CreateStatic(&FCEEditorClonerEffectorExtensionDetailCustomization::OnFilterEffectorActor);
+
+	EffectorsArrayBuilder->OnGenerateArrayElementWidget(FOnGenerateArrayElementWidget::CreateLambda(
+		[FilterActorDelegate](TSharedRef<IPropertyHandle> InAttributeHandle, int32 InArrayIndex, IDetailChildrenBuilder& InChildrenBuilder)
+		{
+			InChildrenBuilder.AddCustomBuilder(MakeShared<FCEEditorClonerCustomActorPickerNodeBuilder>(InAttributeHandle, FilterActorDelegate));
+		}));
+
+	InCategoryBuilder.AddCustomBuilder(EffectorsArrayBuilder);
+}
+
+bool FCEEditorClonerEffectorExtensionDetailCustomization::OnFilterEffectorActor(const AActor* InActor)
+{
+	if (!IsValid(InActor))
+	{
+		return false;
+	}
+
+	return !!InActor->FindComponentByClass<UCEEffectorComponent>();
+}
diff --git a/Engine/Plugins/Experimental/ClonerEffector/Source/ClonerEffectorEditor/Private/Cloner/Customizations/CEEditorClonerEffectorExtensionDetailCustomization.h b/Engine/Plugins/Experimental/ClonerEffector/Source/ClonerEffectorEditor/Private/Cloner/Customizations/CEEditorClonerEffectorExtensionDetailCustomization.h
new file mode 100644
--- /dev/null
+++ b/Engine/Plugins/Experimental/ClonerEffector/Source/ClonerEffectorEditor/Private/Cloner/Customizations/CEEditorClonerEffectorExtensionDetailCustomization.h
@@ -0,0 +1,27 @@
+// Copyright Epic Games, Inc. All Rights Reserved.
+
+#pragma once
+
+#include "IDetailCustomization.h"
+#include "Templates/SharedPointer.h"
+
+class AActor;
+class IDetailCategoryBuilder;
+class IPropertyHandle;
+
+/** Used to customize cloner effector extension properties in details panel */
+class FCEEditorClonerEffectorExtensionDetailCustomization : public IDetailCustomization
+{
+public:
+	static TSharedRef<IDetailCustomization> MakeInstance()
+	{
+		return MakeShared<FCEEditorClonerEffectorExtensionDetailCustomization>();
+	}
+
+	static void CustomizeEffectorsProperty(TSharedRef<IPropertyHandle> InProperty, IDetailCategoryBuilder& InCategoryBuilder);
+	static bool OnFilterEffectorActor(const AActor* InActor);
+
+	//~ Begin IDetailCustomization
+	virtual void CustomizeDetails(IDetailLayoutBuilder& InDetailBuilder) override;
+	//~ End IDetailCustomization
+};
diff --git a/Engine/Plugins/Experimental/ClonerEffector/Source/ClonerEffectorEditor/Private/Cloner/Customizations/CEEditorClonerMeshLayoutDetailCustomization.cpp b/Engine/Plugins/Experimental/ClonerEffector/Source/ClonerEffectorEditor/Private/Cloner/Customizations/CEEditorClonerMeshLayoutDetailCustomization.cpp
new file mode 100644
--- /dev/null
+++ b/Engine/Plugins/Experimental/ClonerEffector/Source/ClonerEffectorEditor/Private/Cloner/Customizations/CEEditorClonerMeshLayoutDetailCustomization.cpp
@@ -0,0 +1,61 @@
+// Copyright Epic Games, Inc. All Rights Reserved.
+
+#include "CEEditorClonerMeshLayoutDetailCustomization.h"
+
+#include "CEEditorClonerCustomActorPickerNodeBuilder.h"
+#include "Cloner/Layouts/CEClonerMeshLayout.h"
+#include "Components/SkeletalMeshComponent.h"
+#include "Components/StaticMeshComponent.h"
+#include "DetailCategoryBuilder.h"
+#include "DetailLayoutBuilder.h"
+#include "GameFramework/Actor.h"
+#include "PropertyHandle.h"
+
+void FCEEditorClonerMeshLayoutDetailCustomization::CustomizeDetails(IDetailLayoutBuilder& InDetailBuilder)
+{
+	AssetPropertyHandle = InDetailBuilder.GetProperty(UCEClonerMeshLayout::GetAssetName(), UCEClonerMeshLayout::StaticClass());
+
+	TSharedRef<IPropertyHandle> MeshPropertyHandle = InDetailBuilder.GetProperty(UCEClonerMeshLayout::GetSampleActorWeakName(), UCEClonerMeshLayout::StaticClass());
+
+	if (!MeshPropertyHandle->IsValidHandle())
+	{
+		return;
+	}
+
+	MeshPropertyHandle->MarkHiddenByCustomization();
+
+	IDetailCategoryBuilder& MeshCategoryBuilder = InDetailBuilder.EditCategory(MeshPropertyHandle->GetDefaultCategoryName(), MeshPropertyHandle->GetDefaultCategoryText());
+
+	MeshCategoryBuilder.AddCustomBuilder(MakeShared<FCEEditorClonerCustomActorPickerNodeBuilder>(
+		MeshPropertyHandle
+		, FOnShouldFilterActor::CreateSP(this, &FCEEditorClonerMeshLayoutDetailCustomization::OnFilterMeshActor))
+	);
+}
+
+bool FCEEditorClonerMeshLayoutDetailCustomization::OnFilterMeshActor(const AActor* InActor)
+{
+	if (!IsValid(InActor))
+	{
+		return false;
+	}
+
+	if (AssetPropertyHandle.IsValid() && AssetPropertyHandle->IsValidHandle())
+	{
+		uint8 Value = 0;
+
+		if (AssetPropertyHandle->GetValue(Value) == FPropertyAccess::Success)
+		{
+			if (static_cast<ECEClonerMeshAsset>(Value) == ECEClonerMeshAsset::SkeletalMesh)
+			{
+				return !!InActor->FindComponentByClass<USkeletalMeshComponent>();
+			}
+
+			if (static_cast<ECEClonerMeshAsset>(Value) == ECEClonerMeshAsset::StaticMesh)
+			{
+				return !!InActor->FindComponentByClass<UStaticMeshComponent>();
+			}
+		}
+	}
+
+	return true;
+}
diff --git a/Engine/Plugins/Experimental/ClonerEffector/Source/ClonerEffectorEditor/Private/Cloner/Customizations/CEEditorClonerMeshLayoutDetailCustomization.h b/Engine/Plugins/Experimental/ClonerEffector/Source/ClonerEffectorEditor/Private/Cloner/Customizations/CEEditorClonerMeshLayoutDetailCustomization.h
new file mode 100644
--- /dev/null
+++ b/Engine/Plugins/Experimental/ClonerEffector/Source/ClonerEffectorEditor/Private/Cloner/Customizations/CEEditorClonerMeshLayoutDetailCustomization.h
@@ -0,0 +1,28 @@
+// Copyright Epic Games, Inc. All Rights Reserved.
+
+#pragma once
+
+#include "IDetailCustomization.h"
+#include "Templates/SharedPointer.h"
+
+class AActor;
+class IPropertyHandle;
+
+/** Used to customize cloner mesh layout properties in details panel */
+class FCEEditorClonerMeshLayoutDetailCustomization : public IDetailCustomization
+{
+public:
+	static TSharedRef<IDetailCustomization> MakeInstance()
+	{
+		return MakeShared<FCEEditorClonerMeshLayoutDetailCustomization>();
+	}
+
+	//~ Begin IDetailCustomization
+	virtual void CustomizeDetails(IDetailLayoutBuilder& InDetailBuilder) override;
+	//~ End IDetailCustomization
+
+private:
+	bool OnFilterMeshActor(const AActor* InActor);
+
+	TSharedPtr<IPropertyHandle> AssetPropertyHandle;
+};
diff --git a/Engine/Plugins/Experimental/ClonerEffector/Source/ClonerEffectorEditor/Private/Cloner/Customizations/CEEditorClonerSplineLayoutDetailCustomization.cpp b/Engine/Plugins/Experimental/ClonerEffector/Source/ClonerEffectorEditor/Private/Cloner/Customizations/CEEditorClonerSplineLayoutDetailCustomization.cpp
new file mode 100644
--- /dev/null
+++ b/Engine/Plugins/Experimental/ClonerEffector/Source/ClonerEffectorEditor/Private/Cloner/Customizations/CEEditorClonerSplineLayoutDetailCustomization.cpp
@@ -0,0 +1,38 @@
+// Copyright Epic Games, Inc. All Rights Reserved.
+
+#include "Cloner/Customizations/CEEditorClonerSplineLayoutDetailCustomization.h"
+
+#include "CEEditorClonerCustomActorPickerNodeBuilder.h"
+#include "Cloner/Layouts/CEClonerSplineLayout.h"
+#include "Components/SplineComponent.h"
+#include "DetailCategoryBuilder.h"
+#include "DetailLayoutBuilder.h"
+
+void FCEEditorClonerSplineLayoutDetailCustomization::CustomizeDetails(IDetailLayoutBuilder& InDetailBuilder)
+{
+	TSharedRef<IPropertyHandle> SplinePropertyHandle = InDetailBuilder.GetProperty(UCEClonerSplineLayout::GetSplineActorWeakName(), UCEClonerSplineLayout::StaticClass());
+
+	if (!SplinePropertyHandle->IsValidHandle())
+	{
+		return;
+	}
+
+	SplinePropertyHandle->MarkHiddenByCustomization();
+
+	IDetailCategoryBuilder& SplineCategoryBuilder = InDetailBuilder.EditCategory(SplinePropertyHandle->GetDefaultCategoryName(), SplinePropertyHandle->GetDefaultCategoryText());
+
+	SplineCategoryBuilder.AddCustomBuilder(MakeShared<FCEEditorClonerCustomActorPickerNodeBuilder>(
+		SplinePropertyHandle
+		, FOnShouldFilterActor::CreateStatic(&FCEEditorClonerSplineLayoutDetailCustomization::OnFilterSplineActor))
+	);
+}
+
+bool FCEEditorClonerSplineLayoutDetailCustomization::OnFilterSplineActor(const AActor* InActor)
+{
+	if (!IsValid(InActor))
+	{
+		return false;
+	}
+
+	return !!InActor->FindComponentByClass<USplineComponent>();
+}
diff --git a/Engine/Plugins/Experimental/ClonerEffector/Source/ClonerEffectorEditor/Private/Cloner/Customizations/CEEditorClonerSplineLayoutDetailCustomization.h b/Engine/Plugins/Experimental/ClonerEffector/Source/ClonerEffectorEditor/Private/Cloner/Customizations/CEEditorClonerSplineLayoutDetailCustomization.h
new file mode 100644
--- /dev/null
+++ b/Engine/Plugins/Experimental/ClonerEffector/Source/ClonerEffectorEditor/Private/Cloner/Customizations/CEEditorClonerSplineLayoutDetailCustomization.h
@@ -0,0 +1,25 @@
+// Copyright Epic Games, Inc. All Rights Reserved.
+
+#pragma once
+
+#include "IDetailCustomization.h"
+#include "Templates/SharedPointer.h"
+
+class AActor;
+
+/** Used to customize cloner spline layout properties in details panel */
+class FCEEditorClonerSplineLayoutDetailCustomization : public IDetailCustomization
+{
+public:
+	static TSharedRef<IDetailCustomization> MakeInstance()
+	{
+		return MakeShared<FCEEditorClonerSplineLayoutDetailCustomization>();
+	}
+
+	//~ Begin IDetailCustomization
+	virtual void CustomizeDetails(IDetailLayoutBuilder& InDetailBuilder) override;
+	//~ End IDetailCustomization
+
+private:
+	static bool OnFilterSplineActor(const AActor* InActor);
+};
