MotionDesign : ClonerEffector
- Adjusting default values and updating display names
- Mirroring offset and rotation values for range extension

diff --git a/Engine/Plugins/Experimental/ClonerEffector/Source/ClonerEffector/Private/Cloner/CEClonerActor.cpp b/Engine/Plugins/Experimental/ClonerEffector/Source/ClonerEffector/Private/Cloner/CEClonerActor.cpp
--- a/Engine/Plugins/Experimental/ClonerEffector/Source/ClonerEffector/Private/Cloner/CEClonerActor.cpp
+++ b/Engine/Plugins/Experimental/ClonerEffector/Source/ClonerEffector/Private/Cloner/CEClonerActor.cpp
@@ -194,6 +194,7 @@ void ACEClonerActor::MigrateDeprecatedProperties()
 		if (UCEClonerRangeExtension* RangeExtension = ClonerComponent->FindOrAddExtension<UCEClonerRangeExtension>())
 		{
 			RangeExtension->SetRangeEnabled(bRangeEnabled);
+			RangeExtension->SetRangeMirrored(false);
 			RangeExtension->SetRangeOffsetMin(RangeOffsetMin);
 			RangeExtension->SetRangeOffsetMax(RangeOffsetMax);
 			RangeExtension->SetRangeRotationMin(RangeRotationMin);
@@ -240,7 +241,7 @@ void ACEClonerActor::MigrateDeprecatedProperties()
 					GridLayout->SetCylinderConstraint(PrevGridLayout->GetCylinderConstraint());
 					GridLayout->SetTextureConstraint(PrevGridLayout->GetTextureConstraint());
 					GridLayout->SetTwistAxis(PrevGridLayout->GetTwistAxis());
-					GridLayout->SetTwistFactor(PrevGridLayout->GetTwistFactor());
+					GridLayout->SetTwistFactor(PrevGridLayout->GetTwistFactor() * 100);
 				}
 			}
 			else if (UCEClonerLineLayout* PrevLineLayout = Cast<UCEClonerLineLayout>(LayoutPair.Value))
@@ -308,7 +309,7 @@ void ACEClonerActor::MigrateDeprecatedProperties()
 					HoneycombLayout->SetHeightSpacing(PrevHoneycombLayout->GetHeightSpacing());
 					HoneycombLayout->SetWidthSpacing(PrevHoneycombLayout->GetWidthSpacing());
 					HoneycombLayout->SetTwistAxis(PrevHoneycombLayout->GetTwistAxis());
-					HoneycombLayout->SetTwistFactor(PrevHoneycombLayout->GetTwistFactor());
+					HoneycombLayout->SetTwistFactor(PrevHoneycombLayout->GetTwistFactor() * 100);
 				}
 			}
 			else if (UCEClonerMeshLayout* PrevMeshLayout = Cast<UCEClonerMeshLayout>(LayoutPair.Value))
diff --git a/Engine/Plugins/Experimental/ClonerEffector/Source/ClonerEffector/Private/Cloner/Extensions/CEClonerRangeExtension.cpp b/Engine/Plugins/Experimental/ClonerEffector/Source/ClonerEffector/Private/Cloner/Extensions/CEClonerRangeExtension.cpp
--- a/Engine/Plugins/Experimental/ClonerEffector/Source/ClonerEffector/Private/Cloner/Extensions/CEClonerRangeExtension.cpp
+++ b/Engine/Plugins/Experimental/ClonerEffector/Source/ClonerEffector/Private/Cloner/Extensions/CEClonerRangeExtension.cpp
@@ -15,6 +15,17 @@ void UCEClonerRangeExtension::SetRangeEnabled(bool bInRangeEnabled)
 	UpdateExtensionParameters();
 }
 
+void UCEClonerRangeExtension::SetRangeMirrored(bool bInMirrored)
+{
+	if (bRangeMirrored == bInMirrored)
+	{
+		return;
+	}
+
+	bRangeMirrored = bInMirrored;
+	UpdateExtensionParameters();
+}
+
 void UCEClonerRangeExtension::SetRangeOffsetMin(const FVector& InRangeOffsetMin)
 {
 	if (RangeOffsetMin == InRangeOffsetMin)
@@ -118,6 +129,12 @@ void UCEClonerRangeExtension::OnExtensionParametersChanged(UCEClonerComponent* I
 {
 	Super::OnExtensionParametersChanged(InComponent);
 
+	if (bRangeMirrored)
+	{
+		RangeOffsetMin = -RangeOffsetMax;
+		RangeRotationMin = -1 * RangeRotationMax;
+	}
+
 	RangeScaleUniformMin = FMath::Clamp(RangeScaleUniformMin, UE_KINDA_SMALL_NUMBER, RangeScaleUniformMax);
 	RangeScaleUniformMax = FMath::Max3(RangeScaleUniformMin, RangeScaleUniformMax, UE_KINDA_SMALL_NUMBER);
 
@@ -155,6 +172,7 @@ const TCEPropertyChangeDispatcher<UCEClonerRangeExtension> UCEClonerRangeExtensi
 {
 	/** Range */
 	{ GET_MEMBER_NAME_CHECKED(UCEClonerRangeExtension, bRangeEnabled), &UCEClonerRangeExtension::OnExtensionPropertyChanged },
+	{ GET_MEMBER_NAME_CHECKED(UCEClonerRangeExtension, bRangeMirrored), &UCEClonerRangeExtension::OnExtensionPropertyChanged },
 	{ GET_MEMBER_NAME_CHECKED(UCEClonerRangeExtension, RangeOffsetMin), &UCEClonerRangeExtension::OnExtensionPropertyChanged },
 	{ GET_MEMBER_NAME_CHECKED(UCEClonerRangeExtension, RangeOffsetMax), &UCEClonerRangeExtension::OnExtensionPropertyChanged },
 	{ GET_MEMBER_NAME_CHECKED(UCEClonerRangeExtension, RangeRotationMin), &UCEClonerRangeExtension::OnExtensionPropertyChanged },
diff --git a/Engine/Plugins/Experimental/ClonerEffector/Source/ClonerEffector/Private/Cloner/Layouts/CEClonerGridLayout.cpp b/Engine/Plugins/Experimental/ClonerEffector/Source/ClonerEffector/Private/Cloner/Layouts/CEClonerGridLayout.cpp
--- a/Engine/Plugins/Experimental/ClonerEffector/Source/ClonerEffector/Private/Cloner/Layouts/CEClonerGridLayout.cpp
+++ b/Engine/Plugins/Experimental/ClonerEffector/Source/ClonerEffector/Private/Cloner/Layouts/CEClonerGridLayout.cpp
@@ -203,7 +203,7 @@ void UCEClonerGridLayout::OnLayoutParametersChanged(UCEClonerComponent* InCompon
 	UNiagaraDataInterfaceTexture* TextureSamplerDI = Cast<UNiagaraDataInterfaceTexture>(ExposedParameters.GetDataInterface(ConstraintTextureSamplerVar));
 	TextureSamplerDI->SetTexture(TextureConstraint.Texture.Get());
 
-	InComponent->SetFloatParameter(TEXT("TwistFactor"), TwistFactor);
+	InComponent->SetFloatParameter(TEXT("TwistFactor"), TwistFactor / 100.f);
 
 	static const FNiagaraVariable TwistAxisVar(FNiagaraTypeDefinition(StaticEnum<ENiagaraOrientationAxis>()), TEXT("TwistAxis"));
 	ExposedParameters.SetParameterValue<int32>(static_cast<int32>(TwistAxis), TwistAxisVar);
diff --git a/Engine/Plugins/Experimental/ClonerEffector/Source/ClonerEffector/Private/Cloner/Layouts/CEClonerHoneycombLayout.cpp b/Engine/Plugins/Experimental/ClonerEffector/Source/ClonerEffector/Private/Cloner/Layouts/CEClonerHoneycombLayout.cpp
--- a/Engine/Plugins/Experimental/ClonerEffector/Source/ClonerEffector/Private/Cloner/Layouts/CEClonerHoneycombLayout.cpp
+++ b/Engine/Plugins/Experimental/ClonerEffector/Source/ClonerEffector/Private/Cloner/Layouts/CEClonerHoneycombLayout.cpp
@@ -179,7 +179,7 @@ void UCEClonerHoneycombLayout::OnLayoutParametersChanged(UCEClonerComponent* InC
 	static const FNiagaraVariable HoneycombPlaneVar(FNiagaraTypeDefinition(StaticEnum<ECEClonerPlane>()), TEXT("HoneycombPlane"));
 	ExposedParameters.SetParameterValue<int32>(static_cast<int32>(Plane), HoneycombPlaneVar);
 
-	InComponent->SetFloatParameter(TEXT("TwistFactor"), TwistFactor);
+	InComponent->SetFloatParameter(TEXT("TwistFactor"), TwistFactor / 100.f);
 
 	static const FNiagaraVariable TwistAxisVar(FNiagaraTypeDefinition(StaticEnum<ENiagaraOrientationAxis>()), TEXT("TwistAxis"));
 	ExposedParameters.SetParameterValue<int32>(static_cast<int32>(TwistAxis), TwistAxisVar);
diff --git a/Engine/Plugins/Experimental/ClonerEffector/Source/ClonerEffector/Public/Cloner/Extensions/CEClonerEmitterSpawnExtension.h b/Engine/Plugins/Experimental/ClonerEffector/Source/ClonerEffector/Public/Cloner/Extensions/CEClonerEmitterSpawnExtension.h
--- a/Engine/Plugins/Experimental/ClonerEffector/Source/ClonerEffector/Public/Cloner/Extensions/CEClonerEmitterSpawnExtension.h
+++ b/Engine/Plugins/Experimental/ClonerEffector/Source/ClonerEffector/Public/Cloner/Extensions/CEClonerEmitterSpawnExtension.h
@@ -83,23 +83,23 @@ protected:
 	//~ End UCEClonerExtensionBase
 
 	/** How many times do we spawn clones */
-	UPROPERTY(EditInstanceOnly, BlueprintReadWrite, Setter="SetSpawnLoopMode", Getter="GetSpawnLoopMode", DisplayName="LoopMode", Category="Spawn")
+	UPROPERTY(EditInstanceOnly, BlueprintReadWrite, Setter="SetSpawnLoopMode", Getter="GetSpawnLoopMode", DisplayName="Emission Mode", Category="Spawn")
 	ECEClonerSpawnLoopMode SpawnLoopMode = ECEClonerSpawnLoopMode::Once;
 
 	/** Amount of spawn iterations for clones */
-	UPROPERTY(EditInstanceOnly, BlueprintReadWrite, Setter="SetSpawnLoopIterations", Getter="GetSpawnLoopIterations", DisplayName="LoopIterations", Category="Spawn", meta=(ClampMin="1", EditCondition="SpawnLoopMode == ECEClonerSpawnLoopMode::Multiple", EditConditionHides))
+	UPROPERTY(EditInstanceOnly, BlueprintReadWrite, Setter="SetSpawnLoopIterations", Getter="GetSpawnLoopIterations", DisplayName="Emission Count", Category="Spawn", meta=(ClampMin="1", EditCondition="SpawnLoopMode == ECEClonerSpawnLoopMode::Multiple", EditConditionHides))
 	int32 SpawnLoopIterations = 1;
 
 	/** Interval/Duration of spawn for clones */
-	UPROPERTY(EditInstanceOnly, BlueprintReadWrite, Setter="SetSpawnLoopInterval", Getter="GetSpawnLoopInterval", DisplayName="LoopInterval", Category="Spawn", meta=(ClampMin="0"))
+	UPROPERTY(EditInstanceOnly, BlueprintReadWrite, Setter="SetSpawnLoopInterval", Getter="GetSpawnLoopInterval", DisplayName="Emission Interval", Category="Spawn", meta=(ClampMin="0", EditCondition="SpawnLoopMode != ECEClonerSpawnLoopMode::Once", EditConditionHides))
 	float SpawnLoopInterval = 1.f;
 
 	/** How does spawn occurs */
-	UPROPERTY(EditInstanceOnly, BlueprintReadWrite, Setter="SetSpawnBehaviorMode", Getter="GetSpawnBehaviorMode", DisplayName="BehaviorMode", Category="Spawn")
+	UPROPERTY(EditInstanceOnly, BlueprintReadWrite, Setter="SetSpawnBehaviorMode", Getter="GetSpawnBehaviorMode", DisplayName="Emission Style", Category="Spawn", meta=(EditCondition="SpawnLoopMode != ECEClonerSpawnLoopMode::Once", EditConditionHides))
 	ECEClonerSpawnBehaviorMode SpawnBehaviorMode = ECEClonerSpawnBehaviorMode::Instant;
 
 	/** How many clones to spawn each seconds */
-	UPROPERTY(EditInstanceOnly, BlueprintReadWrite, Setter="SetSpawnRate", Getter="GetSpawnRate", DisplayName="Rate", Category="Spawn", meta=(ClampMin="0", EditCondition="SpawnBehaviorMode == ECEClonerSpawnBehaviorMode::Rate", EditConditionHides))
+	UPROPERTY(EditInstanceOnly, BlueprintReadWrite, Setter="SetSpawnRate", Getter="GetSpawnRate", DisplayName="Emission Rate", Category="Spawn", meta=(ClampMin="0", EditCondition="SpawnLoopMode != ECEClonerSpawnLoopMode::Once && SpawnBehaviorMode == ECEClonerSpawnBehaviorMode::Rate", EditConditionHides))
 	float SpawnRate = 1.f;
 
 private:
diff --git a/Engine/Plugins/Experimental/ClonerEffector/Source/ClonerEffector/Public/Cloner/Extensions/CEClonerRangeExtension.h b/Engine/Plugins/Experimental/ClonerEffector/Source/ClonerEffector/Public/Cloner/Extensions/CEClonerRangeExtension.h
--- a/Engine/Plugins/Experimental/ClonerEffector/Source/ClonerEffector/Public/Cloner/Extensions/CEClonerRangeExtension.h
+++ b/Engine/Plugins/Experimental/ClonerEffector/Source/ClonerEffector/Public/Cloner/Extensions/CEClonerRangeExtension.h
@@ -32,6 +32,15 @@ public:
 		return bRangeEnabled;
 	}
 
+	UFUNCTION(BlueprintCallable, Category="Cloner")
+	CLONEREFFECTOR_API void SetRangeMirrored(bool bInMirrored);
+
+	UFUNCTION(BlueprintPure, Category="Cloner")
+	bool GetRangeMirrored() const
+	{
+		return bRangeMirrored;
+	}
+
 	UFUNCTION(BlueprintCallable, Category="Cloner")
 	CLONEREFFECTOR_API void SetRangeOffsetMin(const FVector& InRangeOffsetMin);
 
@@ -128,13 +137,17 @@ protected:
 	UPROPERTY(EditInstanceOnly, BlueprintReadWrite, Setter="SetRangeEnabled", Getter="GetRangeEnabled", DisplayName="Enabled", Category="Range")
 	bool bRangeEnabled = false;
 
-	UPROPERTY(EditInstanceOnly, BlueprintReadWrite, Setter="SetRangeOffsetMin", Getter="GetRangeOffsetMin", DisplayName="OffsetMin", Category="Range", meta=(EditCondition="bRangeEnabled", EditConditionHides))
+	/** Mirrors offset and rotation values */
+	UPROPERTY(EditInstanceOnly, BlueprintReadWrite, Setter="SetRangeMirrored", Getter="GetRangeMirrored", DisplayName="Mirrored", Category="Range", meta=(EditCondition="bRangeEnabled", EditConditionHides))
+	bool bRangeMirrored = true;
+
+	UPROPERTY(EditInstanceOnly, BlueprintReadWrite, Setter="SetRangeOffsetMin", Getter="GetRangeOffsetMin", DisplayName="OffsetMin", Category="Range", meta=(EditCondition="bRangeEnabled && !bRangeMirrored", EditConditionHides))
 	FVector RangeOffsetMin = FVector::ZeroVector;
 
 	UPROPERTY(EditInstanceOnly, BlueprintReadWrite, Setter="SetRangeOffsetMax", Getter="GetRangeOffsetMax", DisplayName="OffsetMax", Category="Range", meta=(EditCondition="bRangeEnabled", EditConditionHides))
 	FVector RangeOffsetMax = FVector::ZeroVector;
 
-	UPROPERTY(EditInstanceOnly, BlueprintReadWrite, Setter="SetRangeRotationMin", Getter="GetRangeRotationMin", DisplayName="RotationMin", Category="Range", meta=(EditCondition="bRangeEnabled", EditConditionHides))
+	UPROPERTY(EditInstanceOnly, BlueprintReadWrite, Setter="SetRangeRotationMin", Getter="GetRangeRotationMin", DisplayName="RotationMin", Category="Range", meta=(EditCondition="bRangeEnabled && !bRangeMirrored", EditConditionHides))
 	FRotator RangeRotationMin = FRotator::ZeroRotator;
 
 	UPROPERTY(EditInstanceOnly, BlueprintReadWrite, Setter="SetRangeRotationMax", Getter="GetRangeRotationMax", DisplayName="RotationMax", Category="Range", meta=(EditCondition="bRangeEnabled", EditConditionHides))
diff --git a/Engine/Plugins/Experimental/ClonerEffector/Source/ClonerEffector/Public/Cloner/Layouts/CEClonerGridLayout.h b/Engine/Plugins/Experimental/ClonerEffector/Source/ClonerEffector/Public/Cloner/Layouts/CEClonerGridLayout.h
--- a/Engine/Plugins/Experimental/ClonerEffector/Source/ClonerEffector/Public/Cloner/Layouts/CEClonerGridLayout.h
+++ b/Engine/Plugins/Experimental/ClonerEffector/Source/ClonerEffector/Public/Cloner/Layouts/CEClonerGridLayout.h
@@ -183,11 +183,11 @@ protected:
 	UPROPERTY(EditInstanceOnly, BlueprintReadWrite, Setter="SetTextureConstraint", Getter="GetTextureConstraint", Category="Layout", meta=(EditCondition="Constraint == ECEClonerGridConstraint::Texture", EditConditionHides))
 	FCEClonerGridConstraintTexture TextureConstraint;
 
-	UPROPERTY(EditInstanceOnly, BlueprintReadWrite, Setter, Getter, Category="Layout")
+	UPROPERTY(EditInstanceOnly, BlueprintReadWrite, Setter, Getter, Category="Layout", meta=(UIMin="0", UIMax="100"))
 	float TwistFactor = 0.f;
 
 	UPROPERTY(EditInstanceOnly, BlueprintReadWrite, Setter, Getter, Category="Layout")
-	ENiagaraOrientationAxis TwistAxis = ENiagaraOrientationAxis::ZAxis;
+	ENiagaraOrientationAxis TwistAxis = ENiagaraOrientationAxis::YAxis;
 
 private:
 #if WITH_EDITOR
diff --git a/Engine/Plugins/Experimental/ClonerEffector/Source/ClonerEffector/Public/Cloner/Layouts/CEClonerHoneycombLayout.h b/Engine/Plugins/Experimental/ClonerEffector/Source/ClonerEffector/Public/Cloner/Layouts/CEClonerHoneycombLayout.h
--- a/Engine/Plugins/Experimental/ClonerEffector/Source/ClonerEffector/Public/Cloner/Layouts/CEClonerHoneycombLayout.h
+++ b/Engine/Plugins/Experimental/ClonerEffector/Source/ClonerEffector/Public/Cloner/Layouts/CEClonerHoneycombLayout.h
@@ -137,11 +137,11 @@ protected:
 	UPROPERTY(EditInstanceOnly, BlueprintReadWrite, Setter="SetHeightSpacing", Getter="GetHeightSpacing", Category="Layout")
 	float HeightSpacing = 105.f;
 
-	UPROPERTY(EditInstanceOnly, BlueprintReadWrite, Setter, Getter, Category="Layout")
+	UPROPERTY(EditInstanceOnly, BlueprintReadWrite, Setter, Getter, Category="Layout", meta=(UIMin="0", UIMax="100"))
 	float TwistFactor = 0.f;
 
 	UPROPERTY(EditInstanceOnly, BlueprintReadWrite, Setter, Getter, Category="Layout")
-	ENiagaraOrientationAxis TwistAxis = ENiagaraOrientationAxis::XAxis;
+	ENiagaraOrientationAxis TwistAxis = ENiagaraOrientationAxis::YAxis;
 
 private:
 #if WITH_EDITOR
