MotionDesign : ClonerEffector
- Create a cloner with attached actors in one shortcut CTRL+ALT+C or using the outliner context menu in Motion Design mode

diff --git a/Engine/Plugins/Experimental/Avalanche/Source/AvalancheEditor/Private/AvaEditorCommands.cpp b/Engine/Plugins/Experimental/Avalanche/Source/AvalancheEditor/Private/AvaEditorCommands.cpp
--- a/Engine/Plugins/Experimental/Avalanche/Source/AvalancheEditor/Private/AvaEditorCommands.cpp
+++ b/Engine/Plugins/Experimental/Avalanche/Source/AvalancheEditor/Private/AvaEditorCommands.cpp
@@ -455,6 +455,12 @@ void FAvaEditorCommands::RegisterClonerEffectorCommands()
 		, "Enable cloners on the selected actors in the level"
 		, EUserInterfaceActionType::Button
 		, FInputChord(EKeys::C, EModifierKey::Shift))
+
+	UI_COMMAND(CreateCloner
+		, "Create Cloner"
+		, "Create new cloner with selected actors in the level"
+		, EUserInterfaceActionType::Button
+		, FInputChord(EKeys::C, EModifierKey::Control | EModifierKey::Alt))
 }
 
 #undef LOCTEXT_NAMESPACE
diff --git a/Engine/Plugins/Experimental/Avalanche/Source/AvalancheEditor/Private/AvaEditorCommands.h b/Engine/Plugins/Experimental/Avalanche/Source/AvalancheEditor/Private/AvaEditorCommands.h
--- a/Engine/Plugins/Experimental/Avalanche/Source/AvalancheEditor/Private/AvaEditorCommands.h
+++ b/Engine/Plugins/Experimental/Avalanche/Source/AvalancheEditor/Private/AvaEditorCommands.h
@@ -105,4 +105,5 @@ public:
 	TSharedPtr<FUICommandInfo> EnableCloners;
 	TSharedPtr<FUICommandInfo> DisableEffectors;
 	TSharedPtr<FUICommandInfo> EnableEffectors;
+	TSharedPtr<FUICommandInfo> CreateCloner;
 };
diff --git a/Engine/Plugins/Experimental/Avalanche/Source/AvalancheEditor/Private/ClonerEffector/AvaClonerEffectorExtension.cpp b/Engine/Plugins/Experimental/Avalanche/Source/AvalancheEditor/Private/ClonerEffector/AvaClonerEffectorExtension.cpp
--- a/Engine/Plugins/Experimental/Avalanche/Source/AvalancheEditor/Private/ClonerEffector/AvaClonerEffectorExtension.cpp
+++ b/Engine/Plugins/Experimental/Avalanche/Source/AvalancheEditor/Private/ClonerEffector/AvaClonerEffectorExtension.cpp
@@ -34,6 +34,9 @@ void FAvaClonerEffectorExtension::BindCommands(const TSharedRef<FUICommandList>&
 
 	ClonerEffectorCommands->MapAction(EditorCommands.EnableCloners
 		, FExecuteAction::CreateSP(this, &FAvaClonerEffectorExtension::EnableCloners, true));
+
+	ClonerEffectorCommands->MapAction(EditorCommands.CreateCloner
+		, FExecuteAction::CreateSP(this, &FAvaClonerEffectorExtension::CreateCloner));
 }
 
 TSet<AActor*> FAvaClonerEffectorExtension::GetSelectedActors() const
@@ -127,4 +130,26 @@ void FAvaClonerEffectorExtension::EnableCloners(bool bInEnable) const
 	}
 }
 
+void FAvaClonerEffectorExtension::CreateCloner() const
+{
+	const FEditorModeTools* ModeTools = GetEditorModeTools();
+	if (!ModeTools)
+	{
+		return;
+	}
+
+	UCEClonerSubsystem* ClonerSubsystem = UCEClonerSubsystem::Get();
+	if (!ClonerSubsystem)
+	{
+		return;
+	}
+
+	TSet<AActor*> SelectedActors = GetSelectedActors();
+	UWorld* World = ModeTools->GetWorld();
+
+	constexpr bool bTransact = true;
+
+	ClonerSubsystem->CreateClonerWithActors(World, SelectedActors, bTransact);
+}
+
 #undef LOCTEXT_NAMESPACE
diff --git a/Engine/Plugins/Experimental/Avalanche/Source/AvalancheEditor/Private/ClonerEffector/AvaClonerEffectorExtension.h b/Engine/Plugins/Experimental/Avalanche/Source/AvalancheEditor/Private/ClonerEffector/AvaClonerEffectorExtension.h
--- a/Engine/Plugins/Experimental/Avalanche/Source/AvalancheEditor/Private/ClonerEffector/AvaClonerEffectorExtension.h
+++ b/Engine/Plugins/Experimental/Avalanche/Source/AvalancheEditor/Private/ClonerEffector/AvaClonerEffectorExtension.h
@@ -25,6 +25,7 @@ private:
 
 	void EnableEffectors(bool bInEnable) const;
 	void EnableCloners(bool bInEnable) const;
+	void CreateCloner() const;
 
 	TSharedRef<FUICommandList> ClonerEffectorCommands;
 };
diff --git a/Engine/Plugins/Experimental/Avalanche/Source/AvalancheEffectorsEditor/Private/Cloner/AvaClonerEditorOutlinerContextMenu.cpp b/Engine/Plugins/Experimental/Avalanche/Source/AvalancheEffectorsEditor/Private/Cloner/AvaClonerEditorOutlinerContextMenu.cpp
--- a/Engine/Plugins/Experimental/Avalanche/Source/AvalancheEffectorsEditor/Private/Cloner/AvaClonerEditorOutlinerContextMenu.cpp
+++ b/Engine/Plugins/Experimental/Avalanche/Source/AvalancheEffectorsEditor/Private/Cloner/AvaClonerEditorOutlinerContextMenu.cpp
@@ -2,7 +2,6 @@
 
 #include "Cloner/AvaClonerEditorOutlinerContextMenu.h"
 
-#include "Cloner/CEClonerComponent.h"
 #include "Cloner/Menus/CEEditorClonerMenuContext.h"
 #include "Cloner/Menus/CEEditorClonerMenuOptions.h"
 #include "GameFramework/Actor.h"
@@ -47,6 +46,7 @@ void FAvaClonerEditorOutlinerContextMenu::OnExtendOutlinerContextMenu(UToolMenu*
 			, ECEEditorClonerMenuType::Disable
 			, ECEEditorClonerMenuType::CreateEffector
 			, ECEEditorClonerMenuType::Convert
+			, ECEEditorClonerMenuType::CreateCloner
 		}
 	);
 	Options.UseTransact(true);
@@ -73,10 +73,7 @@ void FAvaClonerEditorOutlinerContextMenu::GetContextObjects(const UAvaOutlinerIt
 		{
 			if (AActor* Actor = ActorItem->GetActor())
 			{
-				if (Actor->FindComponentByClass<UCEClonerComponent>())
-				{
-					OutObjects.Add(Actor);
-				}
+				OutObjects.Add(Actor);
 			}
 		}
 		// is it a component
@@ -84,10 +81,7 @@ void FAvaClonerEditorOutlinerContextMenu::GetContextObjects(const UAvaOutlinerIt
 		{
 			if (USceneComponent* Component = ComponentItem->GetComponent())
 			{
-				if (Component->IsA<UCEClonerComponent>())
-				{
-					OutObjects.Add(Component);
-				}
+				OutObjects.Add(Component);
 			}
 		}
 	}
diff --git a/Engine/Plugins/Experimental/ClonerEffector/Source/ClonerEffector/Private/Subsystems/CEClonerSubsystem.cpp b/Engine/Plugins/Experimental/ClonerEffector/Source/ClonerEffector/Private/Subsystems/CEClonerSubsystem.cpp
--- a/Engine/Plugins/Experimental/ClonerEffector/Source/ClonerEffector/Private/Subsystems/CEClonerSubsystem.cpp
+++ b/Engine/Plugins/Experimental/ClonerEffector/Source/ClonerEffector/Private/Subsystems/CEClonerSubsystem.cpp
@@ -2,6 +2,7 @@
 
 #include "Subsystems/CEClonerSubsystem.h"
 
+#include "Cloner/CEClonerActor.h"
 #include "Cloner/CEClonerComponent.h"
 #include "Cloner/Extensions/CEClonerEffectorExtension.h"
 #include "Cloner/Extensions/CEClonerExtensionBase.h"
@@ -20,6 +21,7 @@
 #include "UObject/UObjectIterator.h"
 
 #if WITH_EDITOR
+#include "Editor.h"
 #include "ScopedTransaction.h"
 #endif
 
@@ -362,6 +364,72 @@ void UCEClonerSubsystem::CreateLinkedEffector(const TSet<UCEClonerComponent*>& I
 }
 #endif
 
+AActor* UCEClonerSubsystem::CreateClonerWithActors(UWorld* InWorld, const TSet<AActor*>& InActors, bool bInShouldTransact)
+{
+	ACEClonerActor* NewClonerActor = nullptr;
+
+	if (!IsValid(InWorld))
+	{
+		return NewClonerActor;
+	}
+
+#if WITH_EDITOR
+	FScopedTransaction Transaction(LOCTEXT("CreateClonerWithActors", "Create cloner with actors attached"), bInShouldTransact);
+#endif
+
+	FActorSpawnParameters Parameters;
+	Parameters.ObjectFlags = RF_Transactional;
+#if WITH_EDITOR
+	Parameters.bTemporaryEditorActor = false;
+#endif
+
+	NewClonerActor = InWorld->SpawnActor<ACEClonerActor>(Parameters);
+
+	if (NewClonerActor)
+	{
+#if WITH_EDITOR
+		NewClonerActor->Modify();
+#endif
+
+		if (!InActors.IsEmpty())
+		{
+			FVector NewAverageLocation;
+
+			for (AActor* Actor : InActors)
+			{
+				if (IsValid(Actor))
+				{
+					NewAverageLocation += Actor->GetActorLocation() / InActors.Num();
+				}
+			}
+
+			NewClonerActor->SetActorLocation(NewAverageLocation);
+
+			for (AActor* Actor : InActors)
+			{
+				if (IsValid(Actor))
+				{
+#if WITH_EDITOR
+					Actor->Modify();
+#endif
+
+					Actor->AttachToActor(NewClonerActor, FAttachmentTransformRules::KeepWorldTransform);
+				}
+			}
+		}
+
+#if WITH_EDITOR
+		if (GEditor)
+		{
+			GEditor->SelectNone(/** SelectionChange */false, /** DeselectBSP */true);
+			GEditor->SelectActor(NewClonerActor, /** Selected */true, /** Notify */true);
+		}
+#endif
+	}
+
+	return NewClonerActor;
+}
+
 TArray<FName> UCEClonerSubsystem::GetLayoutNames() const
 {
 	TArray<FName> LayoutNames;
diff --git a/Engine/Plugins/Experimental/ClonerEffector/Source/ClonerEffector/Public/Subsystems/CEClonerSubsystem.h b/Engine/Plugins/Experimental/ClonerEffector/Source/ClonerEffector/Public/Subsystems/CEClonerSubsystem.h
--- a/Engine/Plugins/Experimental/ClonerEffector/Source/ClonerEffector/Public/Subsystems/CEClonerSubsystem.h
+++ b/Engine/Plugins/Experimental/ClonerEffector/Source/ClonerEffector/Public/Subsystems/CEClonerSubsystem.h
@@ -7,6 +7,7 @@
 #include "UObject/WeakObjectPtrTemplates.h"
 #include "CEClonerSubsystem.generated.h"
 
+class AActor;
 class UCEClonerExtensionBase;
 
 UCLASS(MinimalAPI)
@@ -78,6 +79,10 @@ public:
 	/** Create cloners linked effector */
 	CLONEREFFECTOR_API void CreateLinkedEffector(const TSet<UCEClonerComponent*>& InCloners);
 #endif
+
+	/** Creates a new cloner with actors attached */
+	CLONEREFFECTOR_API AActor* CreateClonerWithActors(UWorld* InWorld, const TSet<AActor*>& InActors, bool bInShouldTransact);
+
 protected:
 	DECLARE_MULTICAST_DELEGATE(FOnSubsystemInitialized)
 	CLONEREFFECTOR_API static FOnSubsystemInitialized OnSubsystemInitializedDelegate;
diff --git a/Engine/Plugins/Experimental/ClonerEffector/Source/ClonerEffectorEditor/Private/Cloner/Menus/CEEditorClonerMenu.cpp b/Engine/Plugins/Experimental/ClonerEffector/Source/ClonerEffectorEditor/Private/Cloner/Menus/CEEditorClonerMenu.cpp
--- a/Engine/Plugins/Experimental/ClonerEffector/Source/ClonerEffectorEditor/Private/Cloner/Menus/CEEditorClonerMenu.cpp
+++ b/Engine/Plugins/Experimental/ClonerEffector/Source/ClonerEffectorEditor/Private/Cloner/Menus/CEEditorClonerMenu.cpp
@@ -187,6 +187,30 @@ void UE::ClonerEditor::Menu::FillConvertClonerSection(UToolMenu* InMenu, const F
 	);
 }
 
+void UE::ClonerEditor::Menu::FillCreateClonerSection(UToolMenu* InMenu, const FCEEditorClonerMenuData& InMenuData)
+{
+	if (!InMenu
+		|| InMenuData.Context.IsEmpty()
+		|| !InMenuData.Context.ContainsAnyActor())
+	{
+		return;
+	}
+
+	FToolMenuSection* ClonerSection = FindOrAddClonerSection(InMenu);
+
+	check(ClonerSection)
+
+	ClonerSection->AddMenuEntry(
+		TEXT("CreateCloner")
+		, LOCTEXT("CreateCloner.Label", "Create cloner")
+		, LOCTEXT("CreateCloner.Tooltip", "Create cloner with selection attached")
+		, FSlateIcon()
+		, FUIAction(
+			FExecuteAction::CreateLambda(&ExecuteCreateClonerAction, InMenuData)
+		)
+	);
+}
+
 void UE::ClonerEditor::Menu::ExecuteEnableClonerAction(const FCEEditorClonerMenuData& InMenuData, bool bInEnable)
 {
 	UCEClonerSubsystem* Subsystem = UCEClonerSubsystem::Get();
@@ -196,7 +220,7 @@ void UE::ClonerEditor::Menu::ExecuteEnableClonerAction(const FCEEditorClonerMenu
 		return;
 	}
 
-	Subsystem->SetClonersEnabled(InMenuData.Context.GetComponents(), bInEnable, InMenuData.Options.ShouldTransact());
+	Subsystem->SetClonersEnabled(InMenuData.Context.GetCloners(), bInEnable, InMenuData.Options.ShouldTransact());
 }
 
 void UE::ClonerEditor::Menu::ExecuteEnableLevelClonerAction(const FCEEditorClonerMenuData& InMenuData, bool bInEnable)
@@ -221,7 +245,7 @@ void UE::ClonerEditor::Menu::ExecuteCreateClonerEffectorAction(const FCEEditorCl
 		return;
 	}
 
-	Subsystem->CreateLinkedEffector(InMenuData.Context.GetComponents());
+	Subsystem->CreateLinkedEffector(InMenuData.Context.GetCloners());
 }
 
 void UE::ClonerEditor::Menu::ExecuteConvertClonerAction(const FCEEditorClonerMenuData& InMenuData, ECEClonerMeshConversion InToMeshType)
@@ -236,4 +260,16 @@ void UE::ClonerEditor::Menu::ExecuteConvertClonerAction(const FCEEditorClonerMen
 	Subsystem->ConvertCloners(InMenuData.Context.GetEnabledCloners(), InToMeshType);
 }
 
+void UE::ClonerEditor::Menu::ExecuteCreateClonerAction(const FCEEditorClonerMenuData& InMenuData)
+{
+	UCEClonerSubsystem* Subsystem = UCEClonerSubsystem::Get();
+
+	if (InMenuData.Context.IsEmpty() || !Subsystem)
+	{
+		return;
+	}
+
+	Subsystem->CreateClonerWithActors(InMenuData.Context.GetWorld(), InMenuData.Context.GetActors(), InMenuData.Options.ShouldTransact());
+}
+
 #undef LOCTEXT_NAMESPACE
diff --git a/Engine/Plugins/Experimental/ClonerEffector/Source/ClonerEffectorEditor/Private/Cloner/Menus/CEEditorClonerMenu.h b/Engine/Plugins/Experimental/ClonerEffector/Source/ClonerEffectorEditor/Private/Cloner/Menus/CEEditorClonerMenu.h
--- a/Engine/Plugins/Experimental/ClonerEffector/Source/ClonerEffectorEditor/Private/Cloner/Menus/CEEditorClonerMenu.h
+++ b/Engine/Plugins/Experimental/ClonerEffector/Source/ClonerEffectorEditor/Private/Cloner/Menus/CEEditorClonerMenu.h
@@ -34,6 +34,8 @@ namespace UE::ClonerEditor::Menu
 
 	void FillConvertClonerSection(UToolMenu* InMenu, const FCEEditorClonerMenuData& InMenuData);
 
+	void FillCreateClonerSection(UToolMenu* InMenu, const FCEEditorClonerMenuData& InMenuData);
+
 	/** Actions */
 
 	void ExecuteEnableClonerAction(const FCEEditorClonerMenuData& InMenuData, bool bInEnable);
@@ -43,4 +45,6 @@ namespace UE::ClonerEditor::Menu
 	void ExecuteCreateClonerEffectorAction(const FCEEditorClonerMenuData& InMenuData);
 
 	void ExecuteConvertClonerAction(const FCEEditorClonerMenuData& InMenuData, ECEClonerMeshConversion InToMeshType);
+
+	void ExecuteCreateClonerAction(const FCEEditorClonerMenuData& InMenuData);
 }
diff --git a/Engine/Plugins/Experimental/ClonerEffector/Source/ClonerEffectorEditor/Private/Cloner/Menus/CEEditorClonerMenuContext.cpp b/Engine/Plugins/Experimental/ClonerEffector/Source/ClonerEffectorEditor/Private/Cloner/Menus/CEEditorClonerMenuContext.cpp
--- a/Engine/Plugins/Experimental/ClonerEffector/Source/ClonerEffectorEditor/Private/Cloner/Menus/CEEditorClonerMenuContext.cpp
+++ b/Engine/Plugins/Experimental/ClonerEffector/Source/ClonerEffectorEditor/Private/Cloner/Menus/CEEditorClonerMenuContext.cpp
@@ -14,8 +14,10 @@ FCEEditorClonerMenuContext::FCEEditorClonerMenuContext(const TSet<UObject*>& InO
 			continue;
 		}
 
-		if (const AActor* Actor = Cast<AActor>(Object))
+		if (AActor* Actor = Cast<AActor>(Object))
 		{
+			ContextActorsKey.Add(Actor);
+
 			TArray<UCEClonerComponent*> ClonerComponents;
 			Actor->GetComponents(ClonerComponents, /** IncludeChildren */false);
 
@@ -23,20 +25,58 @@ FCEEditorClonerMenuContext::FCEEditorClonerMenuContext(const TSet<UObject*>& InO
 			{
 				if (IsValid(Component))
 				{
-					ContextComponents.Add(Component);
+					ContextComponentsKey.Add(Component);
 				}
 			}
 		}
 		else if (UCEClonerComponent* Component = Cast<UCEClonerComponent>(Object))
 		{
-			ContextComponents.Add(Component);
+			ContextComponentsKey.Add(Component);
+			ContextActorsKey.Add(Component->GetOwner());
 		}
 	}
 }
 
-const TSet<UCEClonerComponent*>& FCEEditorClonerMenuContext::GetComponents() const
+TSet<AActor*> FCEEditorClonerMenuContext::GetActors() const
+{
+	TSet<AActor*> Actors;
+	Actors.Reserve(ContextActorsKey.Num());
+
+	Algo::TransformIf(
+		ContextActorsKey
+		, Actors
+		, [](const TObjectKey<AActor>& InActorKey)
+		{
+			return IsValid(InActorKey.ResolveObjectPtr());
+		}
+		, [](const TObjectKey<AActor>& InActorKey)
+		{
+			return InActorKey.ResolveObjectPtr();
+		}
+	);
+
+	return Actors;
+}
+
+TSet<UCEClonerComponent*> FCEEditorClonerMenuContext::GetCloners() const
 {
-	return ContextComponents;
+	TSet<UCEClonerComponent*> Components;
+	Components.Reserve(ContextComponentsKey.Num());
+
+	Algo::TransformIf(
+		ContextComponentsKey
+		, Components
+		, [](const TObjectKey<UCEClonerComponent>& InComponentKey)
+		{
+			return IsValid(InComponentKey.ResolveObjectPtr());
+		}
+		, [](const TObjectKey<UCEClonerComponent>& InComponentKey)
+		{
+			return InComponentKey.ResolveObjectPtr();
+		}
+	);
+
+	return Components;
 }
 
 TSet<UCEClonerComponent*> FCEEditorClonerMenuContext::GetDisabledCloners() const
@@ -51,9 +91,17 @@ TSet<UCEClonerComponent*> FCEEditorClonerMenuContext::GetEnabledCloners() const
 
 UWorld* FCEEditorClonerMenuContext::GetWorld() const
 {
-	for (const UCEClonerComponent* Component : ContextComponents)
+	for (const TObjectKey<AActor>& ActorKey : ContextActorsKey)
+	{
+		if (const AActor* Actor = ActorKey.ResolveObjectPtr())
+		{
+			return Actor->GetWorld();
+		}
+	}
+
+	for (const TObjectKey<UCEClonerComponent>& ComponentKey : ContextComponentsKey)
 	{
-		if (IsValid(Component))
+		if (const UCEClonerComponent* Component = ComponentKey.ResolveObjectPtr())
 		{
 			return Component->GetWorld();
 		}
@@ -64,12 +112,17 @@ UWorld* FCEEditorClonerMenuContext::GetWorld() const
 
 bool FCEEditorClonerMenuContext::IsEmpty() const
 {
-	return ContextComponents.IsEmpty();
+	return ContextComponentsKey.IsEmpty() && ContextActorsKey.IsEmpty();
 }
 
-bool FCEEditorClonerMenuContext::ContainsAnyComponent() const
+bool FCEEditorClonerMenuContext::ContainsAnyActor() const
 {
-	return !ContextComponents.IsEmpty();
+	return !ContextActorsKey.IsEmpty();
+}
+
+bool FCEEditorClonerMenuContext::ContainsAnyCloner() const
+{
+	return !ContextComponentsKey.IsEmpty();
 }
 
 bool FCEEditorClonerMenuContext::ContainsAnyDisabledCloner() const
@@ -84,11 +137,14 @@ bool FCEEditorClonerMenuContext::ContainsAnyEnabledCloner() const
 
 bool FCEEditorClonerMenuContext::ContainsClonerState(bool bInState) const
 {
-	for (UCEClonerComponent* Component : ContextComponents)
+	for (const TObjectKey<UCEClonerComponent>& ComponentKey : ContextComponentsKey)
 	{
-		if (IsValid(Component) && Component->GetEnabled() == bInState)
+		if (const UCEClonerComponent* Component = ComponentKey.ResolveObjectPtr())
 		{
-			return true;
+			if (Component->GetEnabled() == bInState)
+			{
+				return true;
+			}
 		}
 	}
 
@@ -98,13 +154,16 @@ bool FCEEditorClonerMenuContext::ContainsClonerState(bool bInState) const
 TSet<UCEClonerComponent*> FCEEditorClonerMenuContext::GetStateCloners(bool bInState) const
 {
 	TSet<UCEClonerComponent*> Cloners;
-	Cloners.Reserve(ContextComponents.Num());
+	Cloners.Reserve(ContextComponentsKey.Num());
 
-	for (UCEClonerComponent* Component : ContextComponents)
+	for (const TObjectKey<UCEClonerComponent>& ComponentKey : ContextComponentsKey)
 	{
-		if (IsValid(Component) && Component->GetEnabled() == bInState)
+		if (UCEClonerComponent* Component = ComponentKey.ResolveObjectPtr())
 		{
-			Cloners.Add(Component);
+			if (Component->GetEnabled() == bInState)
+			{
+				Cloners.Add(Component);
+			}
 		}
 	}
 
diff --git a/Engine/Plugins/Experimental/ClonerEffector/Source/ClonerEffectorEditor/Private/Effector/Menus/CEEditorEffectorMenuContext.cpp b/Engine/Plugins/Experimental/ClonerEffector/Source/ClonerEffectorEditor/Private/Effector/Menus/CEEditorEffectorMenuContext.cpp
--- a/Engine/Plugins/Experimental/ClonerEffector/Source/ClonerEffectorEditor/Private/Effector/Menus/CEEditorEffectorMenuContext.cpp
+++ b/Engine/Plugins/Experimental/ClonerEffector/Source/ClonerEffectorEditor/Private/Effector/Menus/CEEditorEffectorMenuContext.cpp
@@ -23,20 +23,36 @@ FCEEditorEffectorMenuContext::FCEEditorEffectorMenuContext(const TSet<UObject*>&
 			{
 				if (IsValid(Component))
 				{
-					ContextComponents.Add(Component);
+					ContextComponentsKey.Add(Component);
 				}
 			}
 		}
 		else if (UCEEffectorComponent* Component = Cast<UCEEffectorComponent>(Object))
 		{
-			ContextComponents.Add(Component);
+			ContextComponentsKey.Add(Component);
 		}
 	}
 }
 
-const TSet<UCEEffectorComponent*>& FCEEditorEffectorMenuContext::GetComponents() const
+TSet<UCEEffectorComponent*> FCEEditorEffectorMenuContext::GetComponents() const
 {
-	return ContextComponents;
+	TSet<UCEEffectorComponent*> Components;
+	Components.Reserve(ContextComponentsKey.Num());
+
+	Algo::TransformIf(
+		ContextComponentsKey
+		, Components
+		, [](const TObjectKey<UCEEffectorComponent>& InComponentKey)
+		{
+			return IsValid(InComponentKey.ResolveObjectPtr());
+		}
+		, [](const TObjectKey<UCEEffectorComponent>& InComponentKey)
+		{
+			return InComponentKey.ResolveObjectPtr();
+		}
+	);
+
+	return Components;
 }
 
 TSet<UCEEffectorComponent*> FCEEditorEffectorMenuContext::GetDisabledEffectors() const
@@ -51,9 +67,9 @@ TSet<UCEEffectorComponent*> FCEEditorEffectorMenuContext::GetEnabledEffectors()
 
 UWorld* FCEEditorEffectorMenuContext::GetWorld() const
 {
-	for (const UCEEffectorComponent* Component : ContextComponents)
+	for (const TObjectKey<UCEEffectorComponent>& ComponentKey : ContextComponentsKey)
 	{
-		if (IsValid(Component))
+		if (const UCEEffectorComponent* Component = ComponentKey.ResolveObjectPtr())
 		{
 			return Component->GetWorld();
 		}
@@ -64,12 +80,12 @@ UWorld* FCEEditorEffectorMenuContext::GetWorld() const
 
 bool FCEEditorEffectorMenuContext::IsEmpty() const
 {
-	return ContextComponents.IsEmpty();
+	return ContextComponentsKey.IsEmpty();
 }
 
 bool FCEEditorEffectorMenuContext::ContainsAnyComponent() const
 {
-	return !ContextComponents.IsEmpty();
+	return !ContextComponentsKey.IsEmpty();
 }
 
 bool FCEEditorEffectorMenuContext::ContainsAnyDisabledEffectors() const
@@ -84,11 +100,14 @@ bool FCEEditorEffectorMenuContext::ContainsAnyEnabledEffectors() const
 
 bool FCEEditorEffectorMenuContext::ContainsEffectorState(bool bInState) const
 {
-	for (const UCEEffectorComponent* Component : ContextComponents)
+	for (const TObjectKey<UCEEffectorComponent>& ComponentKey : ContextComponentsKey)
 	{
-		if (IsValid(Component) && Component->GetEnabled() == bInState)
+		if (const UCEEffectorComponent* Component = ComponentKey.ResolveObjectPtr())
 		{
-			return true;
+			if (Component->GetEnabled() == bInState)
+			{
+				return true;
+			}
 		}
 	}
 
@@ -98,13 +117,16 @@ bool FCEEditorEffectorMenuContext::ContainsEffectorState(bool bInState) const
 TSet<UCEEffectorComponent*> FCEEditorEffectorMenuContext::GetStateEffectors(bool bInState) const
 {
 	TSet<UCEEffectorComponent*> Effectors;
-	Effectors.Reserve(ContextComponents.Num());
+	Effectors.Reserve(ContextComponentsKey.Num());
 
-	for (UCEEffectorComponent* Component : ContextComponents)
+	for (const TObjectKey<UCEEffectorComponent>& ComponentKey : ContextComponentsKey)
 	{
-		if (IsValid(Component) && Component->GetEnabled() == bInState)
+		if (UCEEffectorComponent* Component = ComponentKey.ResolveObjectPtr())
 		{
-			Effectors.Add(Component);
+			if (Component->GetEnabled() == bInState)
+			{
+				Effectors.Add(Component);
+			}
 		}
 	}
 
diff --git a/Engine/Plugins/Experimental/ClonerEffector/Source/ClonerEffectorEditor/Private/Subsystems/CEEditorClonerSubsystem.cpp b/Engine/Plugins/Experimental/ClonerEffector/Source/ClonerEffectorEditor/Private/Subsystems/CEEditorClonerSubsystem.cpp
--- a/Engine/Plugins/Experimental/ClonerEffector/Source/ClonerEffectorEditor/Private/Subsystems/CEEditorClonerSubsystem.cpp
+++ b/Engine/Plugins/Experimental/ClonerEffector/Source/ClonerEffectorEditor/Private/Subsystems/CEEditorClonerSubsystem.cpp
@@ -75,7 +75,7 @@ void UCEEditorClonerSubsystem::FillClonerMenu(UToolMenu* InMenu, const FCEEditor
 		}
 	}
 
-	if (InOptions.IsMenuType(ECEEditorClonerMenuType::CreateEffector) && InContext.ContainsAnyComponent())
+	if (InOptions.IsMenuType(ECEEditorClonerMenuType::CreateEffector) && InContext.ContainsAnyCloner())
 	{
 		FillCreateClonerEffectorSection(InMenu, MenuData);
 	}
@@ -97,6 +97,11 @@ void UCEEditorClonerSubsystem::FillClonerMenu(UToolMenu* InMenu, const FCEEditor
 			FillConvertClonerSection(InMenu, MenuData);
 		}
 	}
+
+	if (InOptions.IsMenuType(ECEEditorClonerMenuType::CreateCloner) && !InContext.ContainsAnyCloner() && InContext.ContainsAnyActor())
+	{
+		FillCreateClonerSection(InMenu, MenuData);
+	}
 }
 
 #undef LOCTEXT_NAMESPACE
diff --git a/Engine/Plugins/Experimental/ClonerEffector/Source/ClonerEffectorEditor/Public/Cloner/Menus/CEEditorClonerMenuContext.h b/Engine/Plugins/Experimental/ClonerEffector/Source/ClonerEffectorEditor/Public/Cloner/Menus/CEEditorClonerMenuContext.h
--- a/Engine/Plugins/Experimental/ClonerEffector/Source/ClonerEffectorEditor/Public/Cloner/Menus/CEEditorClonerMenuContext.h
+++ b/Engine/Plugins/Experimental/ClonerEffector/Source/ClonerEffectorEditor/Public/Cloner/Menus/CEEditorClonerMenuContext.h
@@ -4,6 +4,7 @@
 
 #include "Containers/Set.h"
 #include "UObject/Object.h"
+#include "UObject/ObjectKey.h"
 #include "UObject/WeakObjectPtrTemplates.h"
 
 class AActor;
@@ -17,13 +18,15 @@ struct CLONEREFFECTOREDITOR_API FCEEditorClonerMenuContext
 	FCEEditorClonerMenuContext() {}
 	explicit FCEEditorClonerMenuContext(const TSet<UObject*>& InObjects);
 
-	const TSet<UCEClonerComponent*>& GetComponents() const;
+	TSet<AActor*> GetActors() const;
+	TSet<UCEClonerComponent*> GetCloners() const;
 	TSet<UCEClonerComponent*> GetDisabledCloners() const;
 	TSet<UCEClonerComponent*> GetEnabledCloners() const;
 	UWorld* GetWorld() const;
 
 	bool IsEmpty() const;
-	bool ContainsAnyComponent() const;
+	bool ContainsAnyActor() const;
+	bool ContainsAnyCloner() const;
 	bool ContainsAnyDisabledCloner() const;
 	bool ContainsAnyEnabledCloner() const;
 
@@ -31,5 +34,6 @@ protected:
 	bool ContainsClonerState(bool bInState) const;
 	TSet<UCEClonerComponent*> GetStateCloners(bool bInState) const;
 
-	TSet<UCEClonerComponent*> ContextComponents;
+	TSet<TObjectKey<AActor>> ContextActorsKey;
+	TSet<TObjectKey<UCEClonerComponent>> ContextComponentsKey;
 };
\ No newline at end of file
diff --git a/Engine/Plugins/Experimental/ClonerEffector/Source/ClonerEffectorEditor/Public/Cloner/Menus/CEEditorClonerMenuEnums.h b/Engine/Plugins/Experimental/ClonerEffector/Source/ClonerEffectorEditor/Public/Cloner/Menus/CEEditorClonerMenuEnums.h
--- a/Engine/Plugins/Experimental/ClonerEffector/Source/ClonerEffectorEditor/Public/Cloner/Menus/CEEditorClonerMenuEnums.h
+++ b/Engine/Plugins/Experimental/ClonerEffector/Source/ClonerEffectorEditor/Public/Cloner/Menus/CEEditorClonerMenuEnums.h
@@ -7,8 +7,14 @@
 /** Menu type available */
 enum class ECEEditorClonerMenuType : uint8
 {
+	/** Enable the cloner */
 	Enable = 1 << 0,
+	/** Disable the cloner */
 	Disable = 1 << 1,
+	/** Creates a cloner linked effector */
 	CreateEffector = 1 << 2,
-	Convert = 1 << 3
+	/** Converts cloner to specific mesh */
+	Convert = 1 << 3,
+	/** Creates a cloner and attaches actors below it */
+	CreateCloner = 1 << 4
 };
\ No newline at end of file
diff --git a/Engine/Plugins/Experimental/ClonerEffector/Source/ClonerEffectorEditor/Public/Effector/Menus/CEEditorEffectorMenuContext.h b/Engine/Plugins/Experimental/ClonerEffector/Source/ClonerEffectorEditor/Public/Effector/Menus/CEEditorEffectorMenuContext.h
--- a/Engine/Plugins/Experimental/ClonerEffector/Source/ClonerEffectorEditor/Public/Effector/Menus/CEEditorEffectorMenuContext.h
+++ b/Engine/Plugins/Experimental/ClonerEffector/Source/ClonerEffectorEditor/Public/Effector/Menus/CEEditorEffectorMenuContext.h
@@ -4,6 +4,7 @@
 
 #include "Containers/Set.h"
 #include "UObject/Object.h"
+#include "UObject/ObjectKey.h"
 #include "UObject/WeakObjectPtrTemplates.h"
 
 class UCEEffectorComponent;
@@ -16,7 +17,7 @@ struct CLONEREFFECTOREDITOR_API FCEEditorEffectorMenuContext
 	FCEEditorEffectorMenuContext() {}
 	explicit FCEEditorEffectorMenuContext(const TSet<UObject*>& InObjects);
 
-	const TSet<UCEEffectorComponent*>& GetComponents() const;
+	TSet<UCEEffectorComponent*> GetComponents() const;
 	TSet<UCEEffectorComponent*> GetDisabledEffectors() const;
 	TSet<UCEEffectorComponent*> GetEnabledEffectors() const;
 	UWorld* GetWorld() const;
@@ -30,5 +31,5 @@ protected:
 	bool ContainsEffectorState(bool bInState) const;
 	TSet<UCEEffectorComponent*> GetStateEffectors(bool bInState) const;
 
-	TSet<UCEEffectorComponent*> ContextComponents;
+	TSet<TObjectKey<UCEEffectorComponent>> ContextComponentsKey;
 };
\ No newline at end of file
