MotionDesign : ClonerEffector
- Splitting effector forces into multiple niagara functions
- Splitting effector modes into multiple niagara functions
- Adding support for non uniform scale for all effector shapes (types)
- Adding support for drag and vector noise forces in effector
- Fixing effector rotation to go past the limit 180
- Fixing cloner lifetime curve scale

#virtualized

diff --git a/Engine/Plugins/Experimental/ClonerEffector/Source/ClonerEffector/Private/CEClonerEffectorShared.cpp b/Engine/Plugins/Experimental/ClonerEffector/Source/ClonerEffector/Private/CEClonerEffectorShared.cpp
--- a/Engine/Plugins/Experimental/ClonerEffector/Source/ClonerEffector/Private/CEClonerEffectorShared.cpp
+++ b/Engine/Plugins/Experimental/ClonerEffector/Source/ClonerEffector/Private/CEClonerEffectorShared.cpp
@@ -39,7 +39,7 @@ void FCEClonerEffectorChannelData::Write(UNiagaraDataChannelWriter* InWriter) co
 	InWriter->WriteVector(InnerExtentName, Identifier, InnerExtent);
 	InWriter->WriteVector(OuterExtentName, Identifier, OuterExtent);
 	InWriter->WriteVector(LocationDeltaName, Identifier, LocationDelta);
-	InWriter->WriteQuat(RotationDeltaName, Identifier, RotationDelta);
+	InWriter->WriteVector(RotationDeltaName, Identifier, RotationDelta);
 	InWriter->WriteVector(ScaleDeltaName, Identifier, ScaleDelta);
 	InWriter->WritePosition(LocationName, Identifier, Location);
 	InWriter->WriteQuat(RotationName, Identifier, Rotation);
@@ -59,6 +59,9 @@ void FCEClonerEffectorChannelData::Write(UNiagaraDataChannelWriter* InWriter) co
 	InWriter->WriteFloat(AttractionForceStrengthName, Identifier, AttractionForceStrength);
 	InWriter->WriteFloat(AttractionForceFalloffName, Identifier, AttractionForceFalloff);
 	InWriter->WriteVector(GravityForceAccelerationName, Identifier, GravityForceAcceleration);
+	InWriter->WriteFloat(DragForceLinearName, Identifier, DragForceLinear);
+	InWriter->WriteFloat(DragForceRotationalName, Identifier, DragForceRotational);
+	InWriter->WriteFloat(VectorNoiseForceAmountName, Identifier, VectorNoiseForceAmount);
 }
 
 void FCEClonerEffectorChannelData::Read(const UNiagaraDataChannelReader* InReader)
@@ -77,7 +80,7 @@ void FCEClonerEffectorChannelData::Read(const UNiagaraDataChannelReader* InReade
 	InnerExtent = InReader->ReadVector(InnerExtentName, Identifier, bIsValid);
 	OuterExtent = InReader->ReadVector(OuterExtentName, Identifier, bIsValid);
 	LocationDelta = InReader->ReadVector(LocationDeltaName, Identifier, bIsValid);
-	RotationDelta = InReader->ReadQuat(RotationDeltaName, Identifier, bIsValid);
+	RotationDelta = InReader->ReadVector(RotationDeltaName, Identifier, bIsValid);
 	ScaleDelta = InReader->ReadVector(ScaleDeltaName, Identifier, bIsValid);
 	Location = InReader->ReadPosition(LocationName, Identifier, bIsValid);
 	Rotation = InReader->ReadQuat(RotationName, Identifier, bIsValid);
@@ -97,6 +100,9 @@ void FCEClonerEffectorChannelData::Read(const UNiagaraDataChannelReader* InReade
 	AttractionForceStrength = InReader->ReadFloat(AttractionForceStrengthName, Identifier, bIsValid);
 	AttractionForceFalloff = InReader->ReadFloat(AttractionForceFalloffName, Identifier, bIsValid);
 	GravityForceAcceleration = InReader->ReadVector(GravityForceAccelerationName, Identifier, bIsValid);
+	DragForceLinear = InReader->ReadFloat(DragForceLinearName, Identifier, bIsValid);
+	DragForceRotational = InReader->ReadFloat(DragForceRotationalName, Identifier, bIsValid);
+	VectorNoiseForceAmount = InReader->ReadFloat(VectorNoiseForceAmountName, Identifier, bIsValid);
 }
 
 FCEClonerEffectorDataInterfaces::FCEClonerEffectorDataInterfaces(const UNiagaraSystem* InSystem)
diff --git a/Engine/Plugins/Experimental/ClonerEffector/Source/ClonerEffector/Private/Effector/Extensions/CEEffectorForceExtension.cpp b/Engine/Plugins/Experimental/ClonerEffector/Source/ClonerEffector/Private/Effector/Extensions/CEEffectorForceExtension.cpp
--- a/Engine/Plugins/Experimental/ClonerEffector/Source/ClonerEffector/Private/Effector/Extensions/CEEffectorForceExtension.cpp
+++ b/Engine/Plugins/Experimental/ClonerEffector/Source/ClonerEffector/Private/Effector/Extensions/CEEffectorForceExtension.cpp
@@ -60,6 +60,61 @@ void UCEEffectorForceExtension::SetGravityForceAcceleration(const FVector& InAcc
 	UpdateExtensionParameters();
 }
 
+void UCEEffectorForceExtension::SetDragForceEnabled(bool bInEnabled)
+{
+	if (bDragForceEnabled == bInEnabled)
+	{
+		return;
+	}
+
+	bDragForceEnabled = bInEnabled;
+	UpdateExtensionParameters(/** UpdateCloners */ true);
+}
+
+void UCEEffectorForceExtension::SetDragForceLinear(float InStrength)
+{
+	if (FMath::IsNearlyEqual(DragForceLinear, InStrength))
+	{
+		return;
+	}
+
+	DragForceLinear = InStrength;
+	UpdateExtensionParameters();
+}
+
+void UCEEffectorForceExtension::SetDragForceRotational(float InStrength)
+{
+	if (FMath::IsNearlyEqual(DragForceRotational, InStrength))
+	{
+		return;
+	}
+
+	DragForceRotational = InStrength;
+	UpdateExtensionParameters();
+}
+
+void UCEEffectorForceExtension::SetVectorNoiseForceEnabled(bool bInEnabled)
+{
+	if (bVectorNoiseForceEnabled == bInEnabled)
+	{
+		return;
+	}
+
+	bVectorNoiseForceEnabled = bInEnabled;
+	UpdateExtensionParameters(/** UpdateCloners */ true);
+}
+
+void UCEEffectorForceExtension::SetVectorNoiseForceAmount(float InAmount)
+{
+	if (FMath::IsNearlyEqual(VectorNoiseForceAmount, InAmount))
+	{
+		return;
+	}
+
+	VectorNoiseForceAmount = InAmount;
+	UpdateExtensionParameters();
+}
+
 void UCEEffectorForceExtension::SetOrientationForceEnabled(bool bInForceEnabled)
 {
 	if (bOrientationForceEnabled == bInForceEnabled)
@@ -241,6 +296,26 @@ void UCEEffectorForceExtension::OnExtensionParametersChanged(UCEEffectorComponen
 	{
 		ChannelData.GravityForceAcceleration = FVector::ZeroVector;
 	}
+
+	if (bForcesEnabled && bDragForceEnabled)
+	{
+		ChannelData.DragForceLinear = DragForceLinear;
+		ChannelData.DragForceRotational = DragForceRotational;
+	}
+	else
+	{
+		ChannelData.DragForceLinear = 0.f;
+		ChannelData.DragForceRotational = 0.f;
+	}
+
+	if (bForcesEnabled && bVectorNoiseForceEnabled)
+	{
+		ChannelData.VectorNoiseForceAmount = VectorNoiseForceAmount;
+	}
+	else
+	{
+		ChannelData.VectorNoiseForceAmount = 0.f;
+	}
 }
 
 void UCEEffectorForceExtension::OnForceOptionsChanged()
@@ -268,6 +343,11 @@ const TCEPropertyChangeDispatcher<UCEEffectorForceExtension> UCEEffectorForceExt
 	{ GET_MEMBER_NAME_CHECKED(UCEEffectorForceExtension, AttractionForceFalloff), &UCEEffectorForceExtension::OnForceOptionsChanged },
 	{ GET_MEMBER_NAME_CHECKED(UCEEffectorForceExtension, bGravityForceEnabled), &UCEEffectorForceExtension::OnForceOptionsChanged },
 	{ GET_MEMBER_NAME_CHECKED(UCEEffectorForceExtension, GravityForceAcceleration), &UCEEffectorForceExtension::OnForceOptionsChanged },
+	{ GET_MEMBER_NAME_CHECKED(UCEEffectorForceExtension, bDragForceEnabled), &UCEEffectorForceExtension::OnForceOptionsChanged },
+	{ GET_MEMBER_NAME_CHECKED(UCEEffectorForceExtension, DragForceLinear), &UCEEffectorForceExtension::OnForceOptionsChanged },
+	{ GET_MEMBER_NAME_CHECKED(UCEEffectorForceExtension, DragForceRotational), &UCEEffectorForceExtension::OnForceOptionsChanged },
+	{ GET_MEMBER_NAME_CHECKED(UCEEffectorForceExtension, bVectorNoiseForceEnabled), &UCEEffectorForceExtension::OnForceOptionsChanged },
+	{ GET_MEMBER_NAME_CHECKED(UCEEffectorForceExtension, VectorNoiseForceAmount), &UCEEffectorForceExtension::OnForceOptionsChanged },
 };
 
 void UCEEffectorForceExtension::PostEditChangeProperty(FPropertyChangedEvent& InPropertyChangedEvent)
diff --git a/Engine/Plugins/Experimental/ClonerEffector/Source/ClonerEffector/Private/Effector/Modes/CEEffectorNoiseMode.cpp b/Engine/Plugins/Experimental/ClonerEffector/Source/ClonerEffector/Private/Effector/Modes/CEEffectorNoiseMode.cpp
--- a/Engine/Plugins/Experimental/ClonerEffector/Source/ClonerEffector/Private/Effector/Modes/CEEffectorNoiseMode.cpp
+++ b/Engine/Plugins/Experimental/ClonerEffector/Source/ClonerEffector/Private/Effector/Modes/CEEffectorNoiseMode.cpp
@@ -68,7 +68,7 @@ void UCEEffectorNoiseMode::OnExtensionParametersChanged(UCEEffectorComponent* In
 
 	FCEClonerEffectorChannelData& ChannelData = InComponent->GetChannelData();
 	ChannelData.LocationDelta = LocationStrength;
-	ChannelData.RotationDelta = RotationStrength.Quaternion();
+	ChannelData.RotationDelta = FVector(RotationStrength.Yaw, RotationStrength.Pitch, RotationStrength.Roll);
 	ChannelData.ScaleDelta = ScaleStrength;
 	ChannelData.Frequency = Frequency;
 	ChannelData.Pan = Pan;
diff --git a/Engine/Plugins/Experimental/ClonerEffector/Source/ClonerEffector/Private/Effector/Modes/CEEffectorOffsetMode.cpp b/Engine/Plugins/Experimental/ClonerEffector/Source/ClonerEffector/Private/Effector/Modes/CEEffectorOffsetMode.cpp
--- a/Engine/Plugins/Experimental/ClonerEffector/Source/ClonerEffector/Private/Effector/Modes/CEEffectorOffsetMode.cpp
+++ b/Engine/Plugins/Experimental/ClonerEffector/Source/ClonerEffector/Private/Effector/Modes/CEEffectorOffsetMode.cpp
@@ -57,9 +57,9 @@ void UCEEffectorOffsetMode::OnExtensionParametersChanged(UCEEffectorComponent* I
 	Super::OnExtensionParametersChanged(InComponent);
 
 	FCEClonerEffectorChannelData& ChannelData = InComponent->GetChannelData();
-	ChannelData.LocationDelta = GetOffset();
-	ChannelData.RotationDelta = GetRotation().Quaternion();
-	ChannelData.ScaleDelta = GetScale().ComponentMax(FVector(UE_KINDA_SMALL_NUMBER));
+	ChannelData.LocationDelta = Offset;
+	ChannelData.RotationDelta = FVector(Rotation.Yaw, Rotation.Pitch, Rotation.Roll);
+	ChannelData.ScaleDelta = Scale.ComponentMax(FVector(UE_KINDA_SMALL_NUMBER));
 }
 
 #if WITH_EDITOR
diff --git a/Engine/Plugins/Experimental/ClonerEffector/Source/ClonerEffector/Private/Effector/Modes/CEEffectorPushMode.cpp b/Engine/Plugins/Experimental/ClonerEffector/Source/ClonerEffector/Private/Effector/Modes/CEEffectorPushMode.cpp
--- a/Engine/Plugins/Experimental/ClonerEffector/Source/ClonerEffector/Private/Effector/Modes/CEEffectorPushMode.cpp
+++ b/Engine/Plugins/Experimental/ClonerEffector/Source/ClonerEffector/Private/Effector/Modes/CEEffectorPushMode.cpp
@@ -32,7 +32,7 @@ void UCEEffectorPushMode::OnExtensionParametersChanged(UCEEffectorComponent* InC
 
 	FCEClonerEffectorChannelData& ChannelData = InComponent->GetChannelData();
 	ChannelData.LocationDelta = PushStrength;
-	ChannelData.RotationDelta = FQuat::Identity;
+	ChannelData.RotationDelta = FVector::ZeroVector;
 	ChannelData.ScaleDelta = FVector::OneVector;
 	ChannelData.Pan = FVector(0, static_cast<float>(PushDirection), 0);
 }
diff --git a/Engine/Plugins/Experimental/ClonerEffector/Source/ClonerEffector/Private/Effector/Modes/CEEffectorTargetMode.cpp b/Engine/Plugins/Experimental/ClonerEffector/Source/ClonerEffector/Private/Effector/Modes/CEEffectorTargetMode.cpp
--- a/Engine/Plugins/Experimental/ClonerEffector/Source/ClonerEffector/Private/Effector/Modes/CEEffectorTargetMode.cpp
+++ b/Engine/Plugins/Experimental/ClonerEffector/Source/ClonerEffector/Private/Effector/Modes/CEEffectorTargetMode.cpp
@@ -87,7 +87,7 @@ void UCEEffectorTargetMode::OnExtensionParametersChanged(UCEEffectorComponent* I
 	{
 		FCEClonerEffectorChannelData& ChannelData = InComponent->GetChannelData();
 		ChannelData.LocationDelta = InternalTargetActor->GetActorLocation();
-		ChannelData.RotationDelta = FQuat::Identity;
+		ChannelData.RotationDelta = FVector::ZeroVector;
 		ChannelData.ScaleDelta = FVector::OneVector;
 	}
 }
diff --git a/Engine/Plugins/Experimental/ClonerEffector/Source/ClonerEffector/Private/Effector/Types/CEEffectorSphereType.cpp b/Engine/Plugins/Experimental/ClonerEffector/Source/ClonerEffector/Private/Effector/Types/CEEffectorSphereType.cpp
--- a/Engine/Plugins/Experimental/ClonerEffector/Source/ClonerEffector/Private/Effector/Types/CEEffectorSphereType.cpp
+++ b/Engine/Plugins/Experimental/ClonerEffector/Source/ClonerEffector/Private/Effector/Types/CEEffectorSphereType.cpp
@@ -41,8 +41,8 @@ void UCEEffectorSphereType::OnExtensionParametersChanged(UCEEffectorComponent* I
 	OuterRadius = FMath::Max(OuterRadius, InnerRadius);
 
 	FCEClonerEffectorChannelData& ChannelData = InComponent->GetChannelData();
-	ChannelData.InnerExtent = FVector(InnerRadius);
-	ChannelData.OuterExtent = FVector(OuterRadius);
+	ChannelData.InnerExtent = FVector::ZeroVector;
+	ChannelData.OuterExtent = FVector(InnerRadius, OuterRadius, 0.f);
 }
 
 void UCEEffectorSphereType::OnExtensionVisualizerDirty(int32 InDirtyFlags)
diff --git a/Engine/Plugins/Experimental/ClonerEffector/Source/ClonerEffector/Public/CEClonerEffectorShared.h b/Engine/Plugins/Experimental/ClonerEffector/Source/ClonerEffector/Public/CEClonerEffectorShared.h
--- a/Engine/Plugins/Experimental/ClonerEffector/Source/ClonerEffector/Public/CEClonerEffectorShared.h
+++ b/Engine/Plugins/Experimental/ClonerEffector/Source/ClonerEffector/Public/CEClonerEffectorShared.h
@@ -385,6 +385,9 @@ struct FCEClonerEffectorChannelData
 	static constexpr const TCHAR* AttractionForceStrengthName = TEXT("AttractionForceStrength");
 	static constexpr const TCHAR* AttractionForceFalloffName = TEXT("AttractionForceFalloff");
 	static constexpr const TCHAR* GravityForceAccelerationName = TEXT("GravityForceAcceleration");
+	static constexpr const TCHAR* DragForceLinearName = TEXT("DragForceLinear");
+	static constexpr const TCHAR* DragForceRotationalName = TEXT("DragForceRotational");
+	static constexpr const TCHAR* VectorNoiseForceAmountName = TEXT("VectorNoiseForceAmount");
 
 	int32 GetIdentifier() const
 	{
@@ -399,7 +402,7 @@ struct FCEClonerEffectorChannelData
 	FVector InnerExtent = FVector::ZeroVector;
 	FVector OuterExtent = FVector::ZeroVector;
 	FVector LocationDelta = FVector::ZeroVector;
-	FQuat RotationDelta = FQuat::Identity;
+	FVector RotationDelta = FVector::ZeroVector;
 	FVector ScaleDelta = FVector::OneVector;
 	FVector Location = FVector::ZeroVector;
 	FQuat Rotation = FQuat::Identity;
@@ -419,6 +422,9 @@ struct FCEClonerEffectorChannelData
 	float AttractionForceStrength = 0.f;
 	float AttractionForceFalloff = 0.f;
 	FVector GravityForceAcceleration = FVector::ZeroVector;
+	float DragForceLinear = 0.f;
+	float DragForceRotational = 0.f;
+	float VectorNoiseForceAmount = 0.f;
 
 protected:
 	/** Cached effector identifier to detect a change and update cloners DI */
diff --git a/Engine/Plugins/Experimental/ClonerEffector/Source/ClonerEffector/Public/Effector/Extensions/CEEffectorForceExtension.h b/Engine/Plugins/Experimental/ClonerEffector/Source/ClonerEffector/Public/Effector/Extensions/CEEffectorForceExtension.h
--- a/Engine/Plugins/Experimental/ClonerEffector/Source/ClonerEffector/Public/Effector/Extensions/CEEffectorForceExtension.h
+++ b/Engine/Plugins/Experimental/ClonerEffector/Source/ClonerEffector/Public/Effector/Extensions/CEEffectorForceExtension.h
@@ -168,6 +168,51 @@ public:
 		return GravityForceAcceleration;
 	}
 
+	UFUNCTION(BlueprintCallable, Category="Effector")
+	CLONEREFFECTOR_API void SetDragForceEnabled(bool bInEnabled);
+
+	UFUNCTION(BlueprintPure, Category="Effector")
+	bool GetDragForceEnabled() const
+	{
+		return bDragForceEnabled;
+	}
+
+	UFUNCTION(BlueprintCallable, Category="Effector")
+	CLONEREFFECTOR_API void SetDragForceLinear(float InStrength);
+
+	UFUNCTION(BlueprintPure, Category="Effector")
+	float GetDragForceLinear() const
+	{
+		return DragForceLinear;
+	}
+
+	UFUNCTION(BlueprintCallable, Category="Effector")
+	CLONEREFFECTOR_API void SetDragForceRotational(float InStrength);
+
+	UFUNCTION(BlueprintPure, Category="Effector")
+	float GetDragForceRotational() const
+	{
+		return DragForceRotational;
+	}
+
+	UFUNCTION(BlueprintCallable, Category="Effector")
+	CLONEREFFECTOR_API void SetVectorNoiseForceEnabled(bool bInEnabled);
+
+	UFUNCTION(BlueprintPure, Category="Effector")
+	bool GetVectorNoiseForceEnabled() const
+	{
+		return bVectorNoiseForceEnabled;
+	}
+
+	UFUNCTION(BlueprintCallable, Category="Effector")
+	CLONEREFFECTOR_API void SetVectorNoiseForceAmount(float InAmount);
+
+	UFUNCTION(BlueprintPure, Category="Effector")
+	float GetVectorNoiseForceAmount() const
+	{
+		return VectorNoiseForceAmount;
+	}
+
 protected:
 	//~ Begin UObject
 #if WITH_EDITOR
@@ -203,7 +248,7 @@ protected:
 	bool bVortexForceEnabled = false;
 
 	UPROPERTY(EditInstanceOnly, BlueprintReadWrite, Setter, Getter, Category="Force", meta=(EditCondition="bForcesEnabled && bVortexForceEnabled", EditConditionHides))
-	float VortexForceAmount = 10000.f;
+	float VortexForceAmount = 100.f;
 
 	UPROPERTY(EditInstanceOnly, BlueprintReadWrite, Setter, Getter, Category="Force", meta=(EditCondition="bForcesEnabled && bVortexForceEnabled", EditConditionHides))
 	FVector VortexForceAxis = FVector::ZAxisVector;
@@ -213,7 +258,7 @@ protected:
 	bool bCurlNoiseForceEnabled = false;
 
 	UPROPERTY(EditInstanceOnly, BlueprintReadWrite, Setter, Getter, Category="Force", meta=(EditCondition="bForcesEnabled && bCurlNoiseForceEnabled", EditConditionHides))
-	float CurlNoiseForceStrength = 1000.f;
+	float CurlNoiseForceStrength = 100.f;
 
 	UPROPERTY(EditInstanceOnly, BlueprintReadWrite, Setter, Getter, Category="Force", meta=(EditCondition="bForcesEnabled && bCurlNoiseForceEnabled", EditConditionHides))
 	float CurlNoiseForceFrequency = 10.f;
@@ -223,7 +268,7 @@ protected:
 	bool bAttractionForceEnabled = false;
 
 	UPROPERTY(EditInstanceOnly, BlueprintReadWrite, Setter, Getter, Category="Force", meta=(EditCondition="bForcesEnabled && bAttractionForceEnabled", EditConditionHides))
-	float AttractionForceStrength = 1000.f;
+	float AttractionForceStrength = 100.f;
 
 	UPROPERTY(EditInstanceOnly, BlueprintReadWrite, Setter, Getter, Category="Force", meta=(EditCondition="bForcesEnabled && bAttractionForceEnabled", EditConditionHides))
 	float AttractionForceFalloff = 0.1f;
@@ -235,6 +280,23 @@ protected:
 	UPROPERTY(EditInstanceOnly, BlueprintReadWrite, Setter, Getter, Category="Force", meta=(EditCondition="bForcesEnabled && bGravityForceEnabled", EditConditionHides))
 	FVector GravityForceAcceleration = FVector(0, 0, -980.f);
 
+	/** Enable drag force to decrease particles velocity */
+	UPROPERTY(EditInstanceOnly, BlueprintReadWrite, Setter="SetDragForceEnabled", Getter="GetDragForceEnabled", Category="Force", meta=(EditCondition="bForcesEnabled", EditConditionHides))
+	bool bDragForceEnabled = false;
+
+	UPROPERTY(EditInstanceOnly, BlueprintReadWrite, Setter, Getter, Category="Force", meta=(EditCondition="bForcesEnabled && bDragForceEnabled", EditConditionHides))
+	float DragForceLinear = 0.f;
+
+	UPROPERTY(EditInstanceOnly, BlueprintReadWrite, Setter, Getter, Category="Force", meta=(EditCondition="bForcesEnabled && bDragForceEnabled", EditConditionHides))
+	float DragForceRotational = 0.f;
+
+	/** Enable vector random noise force to add variation to clones behavior */
+	UPROPERTY(EditInstanceOnly, BlueprintReadWrite, Setter="SetVectorNoiseForceEnabled", Getter="GetVectorNoiseForceEnabled", Category="Force", meta=(EditCondition="bForcesEnabled", EditConditionHides))
+	bool bVectorNoiseForceEnabled = false;
+
+	UPROPERTY(EditInstanceOnly, BlueprintReadWrite, Setter, Getter, Category="Force", meta=(EditCondition="bForcesEnabled && bVectorNoiseForceEnabled", EditConditionHides))
+	float VectorNoiseForceAmount = 100.f;
+
 private:
 #if WITH_EDITOR
 	/** Used for PECP */
diff --git a/Engine/Plugins/Experimental/ClonerEffector/Source/ClonerEffector/Public/Effector/Modes/CEEffectorOffsetMode.h b/Engine/Plugins/Experimental/ClonerEffector/Source/ClonerEffector/Public/Effector/Modes/CEEffectorOffsetMode.h
--- a/Engine/Plugins/Experimental/ClonerEffector/Source/ClonerEffector/Public/Effector/Modes/CEEffectorOffsetMode.h
+++ b/Engine/Plugins/Experimental/ClonerEffector/Source/ClonerEffector/Public/Effector/Modes/CEEffectorOffsetMode.h
@@ -62,7 +62,7 @@ protected:
 	FVector Offset = FVector::ZeroVector;
 
 	/** Rotation applied on affected clones */
-	UPROPERTY(EditInstanceOnly, BlueprintReadWrite, Setter, Getter, Category="Mode", meta=(ClampMin="-180", ClampMax="180", UIMin="-180", UIMax="180"))
+	UPROPERTY(EditInstanceOnly, BlueprintReadWrite, Setter, Getter, Category="Mode")
 	FRotator Rotation = FRotator::ZeroRotator;
 
 	/** Scale applied on affected clones */
