MotionDesign : ClonerEffector
- Fixed scrubbing issues due to multiple pre/post edit change event triggered on the component and the object containing the property, solved by not exposing the object property in the component to the details view and adding it through details customization, this will also allow for multi cloner/effector details view support since it was not displaying merged properties view when multiple objects were selected before

- Refactoring editor section within cloner/effector extensions to be set in metadata of classes instead of passing it through constructor

- Reducing cloner parameters and system updates, this will avoid resetting the simulation for properties that do not need it and reduce artifacts / flickering but there will always be some motion ghosting due to the simulation reset needed to trigger parameters refresh

diff --git a/Engine/Plugins/Experimental/ClonerEffector/Source/ClonerEffector/Private/CEClonerEffectorShared.cpp b/Engine/Plugins/Experimental/ClonerEffector/Source/ClonerEffector/Private/CEClonerEffectorShared.cpp
--- a/Engine/Plugins/Experimental/ClonerEffector/Source/ClonerEffector/Private/CEClonerEffectorShared.cpp
+++ b/Engine/Plugins/Experimental/ClonerEffector/Source/ClonerEffector/Private/CEClonerEffectorShared.cpp
@@ -178,6 +178,22 @@ int32 FCEClonerEffectorDataInterfaces::Num() const
 	return IndexSize;
 }
 
+void FCEClonerEffectorDataInterfaces::Commit() const
+{
+	for (const TPair<FName, TObjectPtr<UNiagaraDataInterface>>& DataInterfacePair : DataInterfaces)
+	{
+		if (UNiagaraDataInterface* NiagaraDataInterface = DataInterfacePair.Value.Get())
+		{
+			if (UNiagaraSystem* System = NiagaraDataInterface->GetTypedOuter<UNiagaraSystem>())
+			{
+				FNiagaraUserRedirectionParameterStore& UserParameterStore = System->GetExposedParameters();
+				const FNiagaraVariable DIVar(FNiagaraTypeDefinition(NiagaraDataInterface->GetClass()), DataInterfacePair.Key);
+				UserParameterStore.SetDataInterface(NiagaraDataInterface, DIVar);
+			}
+		}
+	}
+}
+
 UNiagaraDataInterfaceArrayInt32* FCEClonerEffectorDataInterfaces::GetIndexArray() const
 {
 	if (const TObjectPtr<UNiagaraDataInterface>* Array = DataInterfaces.Find(FCEClonerEffectorDataInterfaces::IndexName))
@@ -188,6 +204,31 @@ UNiagaraDataInterfaceArrayInt32* FCEClonerEffectorDataInterfaces::GetIndexArray(
 	return nullptr;
 }
 
+#if WITH_EDITOR
+FCEExtensionSection UE::ClonerEffector::EditorSection::GetExtensionSectionFromClass(UClass* InClass)
+{
+	FCEExtensionSection Section;
+
+	if (!InClass)
+	{
+		return Section;
+	}
+
+	while (InClass && !InClass->HasMetaData(TEXT("Section")))
+	{
+		InClass = InClass->GetSuperClass();
+	}
+
+	if (InClass)
+	{
+		Section.SectionName = FName(InClass->GetMetaData(TEXT("Section")));
+		Section.SectionOrder = InClass->GetIntMetaData(TEXT("Priority"));
+	}
+
+	return Section;
+}
+#endif
+
 AStaticMeshActor* UE::ClonerEffector::Conversion::ConvertClonerToStaticMesh(UCEClonerComponent* InCloner)
 {
 	AStaticMeshActor* NewActor = nullptr;
diff --git a/Engine/Plugins/Experimental/ClonerEffector/Source/ClonerEffector/Private/Cloner/CEClonerActor.cpp b/Engine/Plugins/Experimental/ClonerEffector/Source/ClonerEffector/Private/Cloner/CEClonerActor.cpp
--- a/Engine/Plugins/Experimental/ClonerEffector/Source/ClonerEffector/Private/Cloner/CEClonerActor.cpp
+++ b/Engine/Plugins/Experimental/ClonerEffector/Source/ClonerEffector/Private/Cloner/CEClonerActor.cpp
@@ -70,7 +70,7 @@ ACEClonerActor::ACEClonerActor()
 		}
 #endif
 
-		ClonerComponent->OnClonerInitialized().AddUObject(this, &ACEClonerActor::OnClonerInitialized);
+		UCEClonerComponent::OnClonerInitialized().AddUObject(this, &ACEClonerActor::OnClonerInitialized);
 	}
 }
 
@@ -113,10 +113,13 @@ void ACEClonerActor::PostActorCreated()
 
 void ACEClonerActor::OnClonerInitialized(UCEClonerComponent* InClonerComponent)
 {
+	if (ClonerComponent == InClonerComponent)
+	{
 #if WITH_EDITOR
-	// PostActorCreated is sometimes called after the cloner is initialized
-	FTSTicker::GetCoreTicker().AddTicker(FTickerDelegate::CreateUObject(this, &ACEClonerActor::SpawnDefaultActorAttached));
+		// PostActorCreated is sometimes called after the cloner is initialized
+		FTSTicker::GetCoreTicker().AddTicker(FTickerDelegate::CreateUObject(this, &ACEClonerActor::SpawnDefaultActorAttached));
 #endif
+	}
 }
 
 void ACEClonerActor::MigrateDeprecatedProperties()
diff --git a/Engine/Plugins/Experimental/ClonerEffector/Source/ClonerEffector/Private/Cloner/CEClonerComponent.cpp b/Engine/Plugins/Experimental/ClonerEffector/Source/ClonerEffector/Private/Cloner/CEClonerComponent.cpp
--- a/Engine/Plugins/Experimental/ClonerEffector/Source/ClonerEffector/Private/Cloner/CEClonerComponent.cpp
+++ b/Engine/Plugins/Experimental/ClonerEffector/Source/ClonerEffector/Private/Cloner/CEClonerComponent.cpp
@@ -31,6 +31,10 @@
 #include "Widgets/Notifications/SNotificationList.h"
 #endif
 
+UCEClonerComponent::FOnClonerInitialized UCEClonerComponent::OnClonerInitializedDelegate;
+UCEClonerComponent::FOnClonerLayoutLoaded UCEClonerComponent::OnClonerLayoutLoadedDelegate;
+UCEClonerComponent::FOnClonerMeshUpdated UCEClonerComponent::OnClonerMeshUpdatedDelegate;
+
 DEFINE_LOG_CATEGORY_STATIC(LogCEClonerComponent, Log, All);
 
 #define LOCTEXT_NAMESPACE "CEClonerComponent"
@@ -170,15 +174,6 @@ void UCEClonerComponent::UpdateClonerRenderState()
 #endif
 }
 
-void UCEClonerComponent::RefreshUserParameters() const
-{
-	if (UNiagaraSystem* ActiveSystem = GetAsset())
-	{
-		FNiagaraUserRedirectionParameterStore& UserParameterStore = ActiveSystem->GetExposedParameters();
-		UserParameterStore.PostGenericEditChange();
-	}
-}
-
 void UCEClonerComponent::UpdateClonerAttachmentTree(bool bInReset)
 {
 #if WITH_EDITOR
@@ -1004,15 +999,20 @@ void UCEClonerComponent::ShowMaterialWarning(int32 InMaterialCount)
 	}
 }
 
-FName UCEClonerComponent::GetActiveExtensionsName()
+FName UCEClonerComponent::GetActiveExtensionsPropertyName()
 {
 	return GET_MEMBER_NAME_CHECKED(UCEClonerComponent, ActiveExtensions);
 }
 
-FName UCEClonerComponent::GetActiveLayoutName()
+FName UCEClonerComponent::GetActiveLayoutPropertyName()
 {
 	return GET_MEMBER_NAME_CHECKED(UCEClonerComponent, ActiveLayout);
 }
+
+FName UCEClonerComponent::GetLayoutNamePropertyName()
+{
+	return GET_MEMBER_NAME_CHECKED(UCEClonerComponent, LayoutName);
+}
 #endif // WITH_EDITOR
 
 void UCEClonerComponent::InitializeCloner()
@@ -1055,18 +1055,35 @@ bool UCEClonerComponent::TickCloner(float InDelta)
 	{
 		TreeUpdateDeltaTime += InDelta;
 
+		// Update attachment tree
 		if (TreeUpdateDeltaTime >= TreeUpdateInterval)
 		{
-			TreeUpdateDeltaTime -= TreeUpdateInterval;
+			TreeUpdateDeltaTime -= TreeUpdateInterval != 0.f ? TreeUpdateInterval : TreeUpdateDeltaTime;
 
 			UpdateClonerAttachmentTree();
 			UpdateClonerRenderState();
 		}
 
+		// Update layout parameters
+		if (ActiveLayout && ActiveLayout->IsLayoutDirty())
+		{
+			ActiveLayout->UpdateLayoutParameters();
+		}
+
+		// Update extension parameters
+		for (const TObjectPtr<UCEClonerExtensionBase>& ActiveExtension : ActiveExtensions)
+		{
+			if (ActiveExtension && ActiveExtension->IsExtensionDirty())
+			{
+				ActiveExtension->UpdateExtensionParameters();
+			}
+		}
+
+		// Is a simulation reset needed
 		if (bNeedsRefresh)
 		{
 			bNeedsRefresh = false;
-			RequestClonerUpdate(true);
+			RequestClonerUpdate(/**Immediate*/true);
 		}
 	}
 
@@ -1353,7 +1370,12 @@ void UCEClonerComponent::RequestClonerUpdate(bool bInImmediate)
 	if (bInImmediate)
 	{
 		bNeedsRefresh = false;
-		RefreshUserParameters();
+
+		if (UNiagaraSystem* ActiveSystem = GetAsset())
+		{
+			FNiagaraUserRedirectionParameterStore& UserParameterStore = ActiveSystem->GetExposedParameters();
+			UserParameterStore.PostGenericEditChange();
+		}
 	}
 	else
 	{
@@ -1488,7 +1510,7 @@ UCEClonerLayoutBase* UCEClonerComponent::FindOrAddLayout(FName InLayoutName)
 
 	// Check cached layout instances
 	UCEClonerLayoutBase* NewActiveLayout = nullptr;
-	for (TObjectPtr<UCEClonerLayoutBase>& LayoutInstance : LayoutInstances)
+	for (const TObjectPtr<UCEClonerLayoutBase>& LayoutInstance : LayoutInstances)
 	{
 		if (LayoutInstance && LayoutInstance->GetLayoutName() == InLayoutName)
 		{
@@ -1643,6 +1665,12 @@ void UCEClonerComponent::ActivateLayout(UCEClonerLayoutBase* InLayout)
 		return;
 	}
 
+	// Should match current active layout name
+	if (LayoutName != InLayout->GetLayoutName())
+	{
+		return;
+	}
+
 	// Copy data interfaces to new layout
 	if (ActiveLayout && ActiveLayout->IsLayoutLoaded())
 	{
@@ -1678,7 +1706,7 @@ void UCEClonerComponent::OnActiveLayoutChanged()
 	OnSeedChanged();
 	OnColorChanged();
 
-	Layout->UpdateLayoutParameters();
+	Layout->MarkLayoutDirty();
 
 	TSet<TObjectPtr<UCEClonerExtensionBase>> PrevActiveExtensions(ActiveExtensions);
 	ActiveExtensions.Empty();
@@ -1692,7 +1720,7 @@ void UCEClonerComponent::OnActiveLayoutChanged()
 				Extension->ActivateExtension();
 			}
 
-			Extension->UpdateExtensionParameters();
+			Extension->MarkExtensionDirty();
 
 			ActiveExtensions.Add(Extension);
 			PrevActiveExtensions.Remove(Extension);
diff --git a/Engine/Plugins/Experimental/ClonerEffector/Source/ClonerEffector/Private/Cloner/Extensions/CEClonerCollisionExtension.cpp b/Engine/Plugins/Experimental/ClonerEffector/Source/ClonerEffector/Private/Cloner/Extensions/CEClonerCollisionExtension.cpp
--- a/Engine/Plugins/Experimental/ClonerEffector/Source/ClonerEffector/Private/Cloner/Extensions/CEClonerCollisionExtension.cpp
+++ b/Engine/Plugins/Experimental/ClonerEffector/Source/ClonerEffector/Private/Cloner/Extensions/CEClonerCollisionExtension.cpp
@@ -15,7 +15,7 @@ void UCEClonerCollisionExtension::SetSurfaceCollisionEnabled(bool bInSurfaceColl
 	}
 
 	bSurfaceCollisionEnabled = bInSurfaceCollisionEnabled;
-	UpdateExtensionParameters();
+	MarkExtensionDirty();
 }
 
 void UCEClonerCollisionExtension::SetParticleCollisionEnabled(bool bInParticleCollisionEnabled)
@@ -26,7 +26,7 @@ void UCEClonerCollisionExtension::SetParticleCollisionEnabled(bool bInParticleCo
 	}
 
 	bParticleCollisionEnabled = bInParticleCollisionEnabled;
-	UpdateExtensionParameters();
+	MarkExtensionDirty();
 }
 
 void UCEClonerCollisionExtension::SetCollisionVelocityEnabled(bool bInCollisionVelocityEnabled)
@@ -37,7 +37,7 @@ void UCEClonerCollisionExtension::SetCollisionVelocityEnabled(bool bInCollisionV
 	}
 
 	bCollisionVelocityEnabled = bInCollisionVelocityEnabled;
-	UpdateExtensionParameters();
+	MarkExtensionDirty();
 }
 
 void UCEClonerCollisionExtension::SetCollisionIterations(int32 InCollisionIterations)
@@ -49,7 +49,7 @@ void UCEClonerCollisionExtension::SetCollisionIterations(int32 InCollisionIterat
 	}
 
 	CollisionIterations = InCollisionIterations;
-	UpdateExtensionParameters();
+	MarkExtensionDirty();
 }
 
 void UCEClonerCollisionExtension::SetCollisionGridResolution(int32 InCollisionGridResolution)
@@ -61,7 +61,7 @@ void UCEClonerCollisionExtension::SetCollisionGridResolution(int32 InCollisionGr
 	}
 
 	CollisionGridResolution = InCollisionGridResolution;
-	UpdateExtensionParameters();
+	MarkExtensionDirty();
 }
 
 void UCEClonerCollisionExtension::SetCollisionGridSize(const FVector& InCollisionGridSize)
@@ -73,7 +73,7 @@ void UCEClonerCollisionExtension::SetCollisionGridSize(const FVector& InCollisio
 	}
 
 	CollisionGridSize = NewCollisionGridSize;
-	UpdateExtensionParameters();
+	MarkExtensionDirty();
 }
 
 void UCEClonerCollisionExtension::SetCollisionRadiusMode(ECEClonerCollisionRadiusMode InMode)
@@ -84,7 +84,7 @@ void UCEClonerCollisionExtension::SetCollisionRadiusMode(ECEClonerCollisionRadiu
 	}
 
 	CollisionRadiusMode = InMode;
-	UpdateExtensionParameters();
+	MarkExtensionDirty();
 }
 
 void UCEClonerCollisionExtension::SetMassMin(float InMassMin)
@@ -96,7 +96,7 @@ void UCEClonerCollisionExtension::SetMassMin(float InMassMin)
 	}
 
 	MassMin = InMassMin;
-	UpdateExtensionParameters();
+	MarkExtensionDirty();
 }
 
 void UCEClonerCollisionExtension::SetMassMax(float InMassMax)
@@ -108,7 +108,7 @@ void UCEClonerCollisionExtension::SetMassMax(float InMassMax)
 	}
 
 	MassMax = InMassMax;
-	UpdateExtensionParameters();
+	MarkExtensionDirty();
 }
 
 void UCEClonerCollisionExtension::OnExtensionParametersChanged(UCEClonerComponent* InComponent)
@@ -177,7 +177,7 @@ void UCEClonerCollisionExtension::OnClonerMeshesUpdated()
 {
 	Super::OnClonerMeshesUpdated();
 
-	UpdateExtensionParameters();
+	MarkExtensionDirty();
 }
 
 #if WITH_EDITOR
diff --git a/Engine/Plugins/Experimental/ClonerEffector/Source/ClonerEffector/Private/Cloner/Extensions/CEClonerEffectorExtension.cpp b/Engine/Plugins/Experimental/ClonerEffector/Source/ClonerEffector/Private/Cloner/Extensions/CEClonerEffectorExtension.cpp
--- a/Engine/Plugins/Experimental/ClonerEffector/Source/ClonerEffector/Private/Cloner/Extensions/CEClonerEffectorExtension.cpp
+++ b/Engine/Plugins/Experimental/ClonerEffector/Source/ClonerEffector/Private/Cloner/Extensions/CEClonerEffectorExtension.cpp
@@ -187,7 +187,8 @@ void UCEClonerEffectorExtension::OnEffectorsChanged()
 		EffectorIndexArray.Add(EffectorIndex);
 	}
 
-	UpdateExtensionParameters(/** UpdateCloner */true);
+	// Apply changes
+	EffectorDataInterfaces->Commit();
 }
 
 void UCEClonerEffectorExtension::OnEffectorActorsChanged()
diff --git a/Engine/Plugins/Experimental/ClonerEffector/Source/ClonerEffector/Private/Cloner/Extensions/CEClonerEmitterSpawnExtension.cpp b/Engine/Plugins/Experimental/ClonerEffector/Source/ClonerEffector/Private/Cloner/Extensions/CEClonerEmitterSpawnExtension.cpp
--- a/Engine/Plugins/Experimental/ClonerEffector/Source/ClonerEffector/Private/Cloner/Extensions/CEClonerEmitterSpawnExtension.cpp
+++ b/Engine/Plugins/Experimental/ClonerEffector/Source/ClonerEffector/Private/Cloner/Extensions/CEClonerEmitterSpawnExtension.cpp
@@ -15,7 +15,7 @@ void UCEClonerEmitterSpawnExtension::SetSpawnLoopMode(ECEClonerSpawnLoopMode InM
 	}
 
 	SpawnLoopMode = InMode;
-	UpdateExtensionParameters();
+	MarkExtensionDirty();
 }
 
 void UCEClonerEmitterSpawnExtension::SetSpawnLoopIterations(int32 InIterations)
@@ -31,7 +31,7 @@ void UCEClonerEmitterSpawnExtension::SetSpawnLoopIterations(int32 InIterations)
 	}
 
 	SpawnLoopIterations = InIterations;
-	UpdateExtensionParameters();
+	MarkExtensionDirty();
 }
 
 void UCEClonerEmitterSpawnExtension::SetSpawnLoopInterval(float InInterval)
@@ -47,7 +47,7 @@ void UCEClonerEmitterSpawnExtension::SetSpawnLoopInterval(float InInterval)
 	}
 
 	SpawnLoopInterval = InInterval;
-	UpdateExtensionParameters();
+	MarkExtensionDirty();
 }
 
 void UCEClonerEmitterSpawnExtension::SetSpawnBehaviorMode(ECEClonerSpawnBehaviorMode InMode)
@@ -58,7 +58,7 @@ void UCEClonerEmitterSpawnExtension::SetSpawnBehaviorMode(ECEClonerSpawnBehavior
 	}
 
 	SpawnBehaviorMode = InMode;
-	UpdateExtensionParameters();
+	MarkExtensionDirty();
 }
 
 void UCEClonerEmitterSpawnExtension::SetSpawnRate(float InRate)
@@ -74,7 +74,7 @@ void UCEClonerEmitterSpawnExtension::SetSpawnRate(float InRate)
 	}
 
 	SpawnRate = InRate;
-	UpdateExtensionParameters();
+	MarkExtensionDirty();
 }
 
 void UCEClonerEmitterSpawnExtension::OnExtensionParametersChanged(UCEClonerComponent* InComponent)
diff --git a/Engine/Plugins/Experimental/ClonerEffector/Source/ClonerEffector/Private/Cloner/Extensions/CEClonerExtensionBase.cpp b/Engine/Plugins/Experimental/ClonerEffector/Source/ClonerEffector/Private/Cloner/Extensions/CEClonerExtensionBase.cpp
--- a/Engine/Plugins/Experimental/ClonerEffector/Source/ClonerEffector/Private/Cloner/Extensions/CEClonerExtensionBase.cpp
+++ b/Engine/Plugins/Experimental/ClonerEffector/Source/ClonerEffector/Private/Cloner/Extensions/CEClonerExtensionBase.cpp
@@ -4,6 +4,13 @@
 
 #include "Cloner/CEClonerComponent.h"
 
+#if WITH_EDITOR
+FCEExtensionSection UCEClonerExtensionBase::GetExtensionSection() const
+{
+	return UE::ClonerEffector::EditorSection::GetExtensionSectionFromClass(GetClass());
+}
+#endif
+
 UCEClonerComponent* UCEClonerExtensionBase::GetClonerComponent() const
 {
 	return GetTypedOuter<UCEClonerComponent>();
@@ -51,7 +58,7 @@ void UCEClonerExtensionBase::DeactivateExtension()
 	}
 }
 
-void UCEClonerExtensionBase::UpdateExtensionParameters(bool bInUpdateCloner, bool bInImmediate)
+void UCEClonerExtensionBase::UpdateExtensionParameters()
 {
 	if (!IsExtensionActive())
 	{
@@ -67,18 +74,35 @@ void UCEClonerExtensionBase::UpdateExtensionParameters(bool bInUpdateCloner, boo
 
 		OnExtensionParametersChanged(ClonerComponent);
 
-		if (bInUpdateCloner)
+		if (EnumHasAnyFlags(ExtensionStatus, ECEClonerSystemStatus::SimulationDirty))
 		{
-			ClonerComponent->RequestClonerUpdate(bInImmediate);
+			ClonerComponent->RequestClonerUpdate(/*Immediate*/false);
 		}
 	}
+
+	ExtensionStatus = ECEClonerSystemStatus::UpToDate;
+}
+
+void UCEClonerExtensionBase::MarkExtensionDirty(bool bInUpdateCloner)
+{
+	EnumAddFlags(ExtensionStatus, ECEClonerSystemStatus::ParametersDirty);
+
+	if (bInUpdateCloner)
+	{
+		EnumAddFlags(ExtensionStatus, ECEClonerSystemStatus::SimulationDirty);
+	}
+}
+
+bool UCEClonerExtensionBase::IsExtensionDirty() const
+{
+	return EnumHasAnyFlags(ExtensionStatus, ECEClonerSystemStatus::ParametersDirty);
 }
 
 void UCEClonerExtensionBase::PostEditImport()
 {
 	Super::PostEditImport();
 
-	UpdateExtensionParameters();
+	MarkExtensionDirty();
 }
 
 #if WITH_EDITOR
@@ -86,11 +110,11 @@ void UCEClonerExtensionBase::PostEditUndo()
 {
 	Super::PostEditUndo();
 
-	UpdateExtensionParameters();
+	MarkExtensionDirty();
 }
 #endif
 
 void UCEClonerExtensionBase::OnExtensionPropertyChanged()
 {
-	UpdateExtensionParameters();
+	MarkExtensionDirty();
 }
diff --git a/Engine/Plugins/Experimental/ClonerEffector/Source/ClonerEffector/Private/Cloner/Extensions/CEClonerLifetimeExtension.cpp b/Engine/Plugins/Experimental/ClonerEffector/Source/ClonerEffector/Private/Cloner/Extensions/CEClonerLifetimeExtension.cpp
--- a/Engine/Plugins/Experimental/ClonerEffector/Source/ClonerEffector/Private/Cloner/Extensions/CEClonerLifetimeExtension.cpp
+++ b/Engine/Plugins/Experimental/ClonerEffector/Source/ClonerEffector/Private/Cloner/Extensions/CEClonerLifetimeExtension.cpp
@@ -11,11 +11,7 @@
 UCEClonerLifetimeExtension::UCEClonerLifetimeExtension()
 	: UCEClonerExtensionBase(
 		TEXT("Lifetime")
-		, 0
-#if WITH_EDITOR
-		, UE::ClonerEffector::ClonerSection::EmissionSection
-#endif
-		)
+		, 0)
 {
 	// Default Scale Curve
 	LifetimeScaleCurve.AddKey(0, 1.f);
@@ -30,7 +26,7 @@ void UCEClonerLifetimeExtension::SetLifetimeEnabled(bool bInEnabled)
 	}
 
 	bLifetimeEnabled = bInEnabled;
-	UpdateExtensionParameters();
+	MarkExtensionDirty();
 }
 
 void UCEClonerLifetimeExtension::SetLifetimeMin(float InMin)
@@ -46,7 +42,7 @@ void UCEClonerLifetimeExtension::SetLifetimeMin(float InMin)
 	}
 
 	LifetimeMin = InMin;
-	UpdateExtensionParameters();
+	MarkExtensionDirty();
 }
 
 void UCEClonerLifetimeExtension::SetLifetimeMax(float InMax)
@@ -62,7 +58,7 @@ void UCEClonerLifetimeExtension::SetLifetimeMax(float InMax)
 	}
 
 	LifetimeMax = InMax;
-	UpdateExtensionParameters();
+	MarkExtensionDirty();
 }
 
 void UCEClonerLifetimeExtension::SetLifetimeScaleEnabled(bool bInEnabled)
@@ -73,13 +69,13 @@ void UCEClonerLifetimeExtension::SetLifetimeScaleEnabled(bool bInEnabled)
 	}
 
 	bLifetimeScaleEnabled = bInEnabled;
-	UpdateExtensionParameters();
+	MarkExtensionDirty();
 }
 
 void UCEClonerLifetimeExtension::SetLifetimeScaleCurve(const FRichCurve& InCurve)
 {
 	LifetimeScaleCurve = InCurve;
-	UpdateExtensionParameters();
+	MarkExtensionDirty();
 }
 
 void UCEClonerLifetimeExtension::OnExtensionParametersChanged(UCEClonerComponent* InComponent)
@@ -128,7 +124,7 @@ void UCEClonerLifetimeExtension::OnLifetimeScaleCurveChanged()
 	if (const UNiagaraDataInterfaceCurve* LifetimeCurve = LifetimeScaleCurveDIWeak.Get())
 	{
 		LifetimeScaleCurve = LifetimeCurve->Curve;
-		UpdateExtensionParameters();
+		MarkExtensionDirty();
 	}
 }
 
diff --git a/Engine/Plugins/Experimental/ClonerEffector/Source/ClonerEffector/Private/Cloner/Extensions/CEClonerMeshRendererExtension.cpp b/Engine/Plugins/Experimental/ClonerEffector/Source/ClonerEffector/Private/Cloner/Extensions/CEClonerMeshRendererExtension.cpp
--- a/Engine/Plugins/Experimental/ClonerEffector/Source/ClonerEffector/Private/Cloner/Extensions/CEClonerMeshRendererExtension.cpp
+++ b/Engine/Plugins/Experimental/ClonerEffector/Source/ClonerEffector/Private/Cloner/Extensions/CEClonerMeshRendererExtension.cpp
@@ -14,11 +14,7 @@ DEFINE_LOG_CATEGORY_STATIC(LogCEClonerMeshRenderer, Log, All);
 UCEClonerMeshRendererExtension::UCEClonerMeshRendererExtension()
 	: UCEClonerExtensionBase(
 		TEXT("MeshRenderer")
-		, 1
-#if WITH_EDITOR
-		, UE::ClonerEffector::ClonerSection::RenderingSection
-#endif
-		)
+		, 1)
 {
 	// Default override material
 	static ConstructorHelpers::FObjectFinder<UMaterialInterface> DefaultMaterialFinder(UCEClonerEffectorSettings::DefaultMaterialPath);
@@ -33,7 +29,7 @@ void UCEClonerMeshRendererExtension::SetMeshRenderMode(ECEClonerMeshRenderMode I
 	}
 
 	MeshRenderMode = InMode;
-	UpdateExtensionParameters();
+	MarkExtensionDirty();
 }
 
 void UCEClonerMeshRendererExtension::SetMeshFacingMode(ENiagaraMeshFacingMode InMode)
diff --git a/Engine/Plugins/Experimental/ClonerEffector/Source/ClonerEffector/Private/Cloner/Extensions/CEClonerProgressExtension.cpp b/Engine/Plugins/Experimental/ClonerEffector/Source/ClonerEffector/Private/Cloner/Extensions/CEClonerProgressExtension.cpp
--- a/Engine/Plugins/Experimental/ClonerEffector/Source/ClonerEffector/Private/Cloner/Extensions/CEClonerProgressExtension.cpp
+++ b/Engine/Plugins/Experimental/ClonerEffector/Source/ClonerEffector/Private/Cloner/Extensions/CEClonerProgressExtension.cpp
@@ -12,7 +12,7 @@ void UCEClonerProgressExtension::SetInvertProgress(bool bInInvertProgress)
 	}
 
 	bInvertProgress = bInInvertProgress;
-	UpdateExtensionParameters(/** UpdateCloner */false);
+	MarkExtensionDirty(/** UpdateCloner */false);
 }
 
 void UCEClonerProgressExtension::SetProgress(float InProgress)
@@ -25,7 +25,7 @@ void UCEClonerProgressExtension::SetProgress(float InProgress)
 	}
 
 	Progress = InProgress;
-	UpdateExtensionParameters(/** UpdateCloner */false);
+	MarkExtensionDirty(/** UpdateCloner */false);
 }
 
 void UCEClonerProgressExtension::OnExtensionParametersChanged(UCEClonerComponent* InComponent)
diff --git a/Engine/Plugins/Experimental/ClonerEffector/Source/ClonerEffector/Private/Cloner/Extensions/CEClonerRangeExtension.cpp b/Engine/Plugins/Experimental/ClonerEffector/Source/ClonerEffector/Private/Cloner/Extensions/CEClonerRangeExtension.cpp
--- a/Engine/Plugins/Experimental/ClonerEffector/Source/ClonerEffector/Private/Cloner/Extensions/CEClonerRangeExtension.cpp
+++ b/Engine/Plugins/Experimental/ClonerEffector/Source/ClonerEffector/Private/Cloner/Extensions/CEClonerRangeExtension.cpp
@@ -12,7 +12,7 @@ void UCEClonerRangeExtension::SetRangeEnabled(bool bInRangeEnabled)
 	}
 
 	bRangeEnabled = bInRangeEnabled;
-	UpdateExtensionParameters();
+	MarkExtensionDirty();
 }
 
 void UCEClonerRangeExtension::SetRangeMirrored(bool bInMirrored)
@@ -23,7 +23,7 @@ void UCEClonerRangeExtension::SetRangeMirrored(bool bInMirrored)
 	}
 
 	bRangeMirrored = bInMirrored;
-	UpdateExtensionParameters();
+	MarkExtensionDirty();
 }
 
 void UCEClonerRangeExtension::SetRangeOffsetMin(const FVector& InRangeOffsetMin)
@@ -34,7 +34,7 @@ void UCEClonerRangeExtension::SetRangeOffsetMin(const FVector& InRangeOffsetMin)
 	}
 
 	RangeOffsetMin = InRangeOffsetMin;
-	UpdateExtensionParameters();
+	MarkExtensionDirty();
 }
 
 void UCEClonerRangeExtension::SetRangeOffsetMax(const FVector& InRangeOffsetMax)
@@ -45,7 +45,7 @@ void UCEClonerRangeExtension::SetRangeOffsetMax(const FVector& InRangeOffsetMax)
 	}
 
 	RangeOffsetMax = InRangeOffsetMax;
-	UpdateExtensionParameters();
+	MarkExtensionDirty();
 }
 
 void UCEClonerRangeExtension::SetRangeRotationMin(const FRotator& InRangeRotationMin)
@@ -56,7 +56,7 @@ void UCEClonerRangeExtension::SetRangeRotationMin(const FRotator& InRangeRotatio
 	}
 
 	RangeRotationMin = InRangeRotationMin;
-	UpdateExtensionParameters();
+	MarkExtensionDirty();
 }
 
 void UCEClonerRangeExtension::SetRangeRotationMax(const FRotator& InRangeRotationMax)
@@ -67,7 +67,7 @@ void UCEClonerRangeExtension::SetRangeRotationMax(const FRotator& InRangeRotatio
 	}
 
 	RangeRotationMax = InRangeRotationMax;
-	UpdateExtensionParameters();
+	MarkExtensionDirty();
 }
 
 void UCEClonerRangeExtension::SetRangeScaleUniform(bool bInRangeScaleUniform)
@@ -78,7 +78,7 @@ void UCEClonerRangeExtension::SetRangeScaleUniform(bool bInRangeScaleUniform)
 	}
 
 	bRangeScaleUniform = bInRangeScaleUniform;
-	UpdateExtensionParameters();
+	MarkExtensionDirty();
 }
 
 void UCEClonerRangeExtension::SetRangeScaleMin(const FVector& InRangeScaleMin)
@@ -89,7 +89,7 @@ void UCEClonerRangeExtension::SetRangeScaleMin(const FVector& InRangeScaleMin)
 	}
 
 	RangeScaleMin = InRangeScaleMin;
-	UpdateExtensionParameters();
+	MarkExtensionDirty();
 }
 
 void UCEClonerRangeExtension::SetRangeScaleMax(const FVector& InRangeScaleMax)
@@ -100,7 +100,7 @@ void UCEClonerRangeExtension::SetRangeScaleMax(const FVector& InRangeScaleMax)
 	}
 
 	RangeScaleMax = InRangeScaleMax;
-	UpdateExtensionParameters();
+	MarkExtensionDirty();
 }
 
 void UCEClonerRangeExtension::SetRangeScaleUniformMin(float InRangeScaleUniformMin)
@@ -111,7 +111,7 @@ void UCEClonerRangeExtension::SetRangeScaleUniformMin(float InRangeScaleUniformM
 	}
 
 	RangeScaleUniformMin = InRangeScaleUniformMin;
-	UpdateExtensionParameters();
+	MarkExtensionDirty();
 }
 
 void UCEClonerRangeExtension::SetRangeScaleUniformMax(float InRangeScaleUniformMax)
@@ -122,7 +122,7 @@ void UCEClonerRangeExtension::SetRangeScaleUniformMax(float InRangeScaleUniformM
 	}
 
 	RangeScaleUniformMax = InRangeScaleUniformMax;
-	UpdateExtensionParameters();
+	MarkExtensionDirty();
 }
 
 void UCEClonerRangeExtension::OnExtensionParametersChanged(UCEClonerComponent* InComponent)
diff --git a/Engine/Plugins/Experimental/ClonerEffector/Source/ClonerEffector/Private/Cloner/Extensions/CEClonerStepExtension.cpp b/Engine/Plugins/Experimental/ClonerEffector/Source/ClonerEffector/Private/Cloner/Extensions/CEClonerStepExtension.cpp
--- a/Engine/Plugins/Experimental/ClonerEffector/Source/ClonerEffector/Private/Cloner/Extensions/CEClonerStepExtension.cpp
+++ b/Engine/Plugins/Experimental/ClonerEffector/Source/ClonerEffector/Private/Cloner/Extensions/CEClonerStepExtension.cpp
@@ -12,7 +12,7 @@ void UCEClonerStepExtension::SetDeltaStepEnabled(bool bInEnabled)
 	}
 
 	bDeltaStepEnabled = bInEnabled;
-	UpdateExtensionParameters();
+	MarkExtensionDirty();
 }
 
 void UCEClonerStepExtension::SetDeltaStepPosition(const FVector& InPosition)
@@ -23,7 +23,7 @@ void UCEClonerStepExtension::SetDeltaStepPosition(const FVector& InPosition)
 	}
 
 	DeltaStepPosition = InPosition;
-	UpdateExtensionParameters();
+	MarkExtensionDirty();
 }
 
 void UCEClonerStepExtension::SetDeltaStepRotation(const FRotator& InRotation)
@@ -34,7 +34,7 @@ void UCEClonerStepExtension::SetDeltaStepRotation(const FRotator& InRotation)
 	}
 
 	DeltaStepRotation = InRotation;
-	UpdateExtensionParameters();
+	MarkExtensionDirty();
 }
 
 void UCEClonerStepExtension::SetDeltaStepScale(const FVector& InScale)
@@ -45,7 +45,7 @@ void UCEClonerStepExtension::SetDeltaStepScale(const FVector& InScale)
 	}
 
 	DeltaStepScale = InScale;
-	UpdateExtensionParameters();
+	MarkExtensionDirty();
 }
 
 void UCEClonerStepExtension::OnExtensionParametersChanged(UCEClonerComponent* InComponent)
diff --git a/Engine/Plugins/Experimental/ClonerEffector/Source/ClonerEffector/Private/Cloner/Layouts/CEClonerCircleLayout.cpp b/Engine/Plugins/Experimental/ClonerEffector/Source/ClonerEffector/Private/Cloner/Layouts/CEClonerCircleLayout.cpp
--- a/Engine/Plugins/Experimental/ClonerEffector/Source/ClonerEffector/Private/Cloner/Layouts/CEClonerCircleLayout.cpp
+++ b/Engine/Plugins/Experimental/ClonerEffector/Source/ClonerEffector/Private/Cloner/Layouts/CEClonerCircleLayout.cpp
@@ -12,7 +12,7 @@ void UCEClonerCircleLayout::SetCount(int32 InCount)
 	}
 
 	Count = InCount;
-	UpdateLayoutParameters();
+	MarkLayoutDirty();
 }
 
 void UCEClonerCircleLayout::SetRadius(float InRadius)
@@ -23,7 +23,7 @@ void UCEClonerCircleLayout::SetRadius(float InRadius)
 	}
 
 	Radius = InRadius;
-	UpdateLayoutParameters();
+	MarkLayoutDirty();
 }
 
 void UCEClonerCircleLayout::SetAngleStart(float InAngleStart)
@@ -34,7 +34,7 @@ void UCEClonerCircleLayout::SetAngleStart(float InAngleStart)
 	}
 
 	AngleStart = InAngleStart;
-	UpdateLayoutParameters();
+	MarkLayoutDirty();
 }
 
 void UCEClonerCircleLayout::SetAngleRatio(float InAngleRatio)
@@ -45,7 +45,7 @@ void UCEClonerCircleLayout::SetAngleRatio(float InAngleRatio)
 	}
 
 	AngleRatio = InAngleRatio;
-	UpdateLayoutParameters();
+	MarkLayoutDirty();
 }
 
 void UCEClonerCircleLayout::SetOrientMesh(bool bInOrientMesh)
@@ -56,7 +56,7 @@ void UCEClonerCircleLayout::SetOrientMesh(bool bInOrientMesh)
 	}
 
 	bOrientMesh = bInOrientMesh;
-	UpdateLayoutParameters();
+	MarkLayoutDirty();
 }
 
 void UCEClonerCircleLayout::SetPlane(ECEClonerPlane InPlane)
@@ -67,7 +67,7 @@ void UCEClonerCircleLayout::SetPlane(ECEClonerPlane InPlane)
 	}
 
 	Plane = InPlane;
-	UpdateLayoutParameters();
+	MarkLayoutDirty();
 }
 
 void UCEClonerCircleLayout::SetRotation(const FRotator& InRotation)
@@ -78,7 +78,7 @@ void UCEClonerCircleLayout::SetRotation(const FRotator& InRotation)
 	}
 
 	Rotation = InRotation;
-	UpdateLayoutParameters();
+	MarkLayoutDirty();
 }
 
 void UCEClonerCircleLayout::SetScale(const FVector& InScale)
@@ -89,7 +89,7 @@ void UCEClonerCircleLayout::SetScale(const FVector& InScale)
 	}
 
 	Scale = InScale;
-	UpdateLayoutParameters();
+	MarkLayoutDirty();
 }
 
 #if WITH_EDITOR
diff --git a/Engine/Plugins/Experimental/ClonerEffector/Source/ClonerEffector/Private/Cloner/Layouts/CEClonerCylinderLayout.cpp b/Engine/Plugins/Experimental/ClonerEffector/Source/ClonerEffector/Private/Cloner/Layouts/CEClonerCylinderLayout.cpp
--- a/Engine/Plugins/Experimental/ClonerEffector/Source/ClonerEffector/Private/Cloner/Layouts/CEClonerCylinderLayout.cpp
+++ b/Engine/Plugins/Experimental/ClonerEffector/Source/ClonerEffector/Private/Cloner/Layouts/CEClonerCylinderLayout.cpp
@@ -12,7 +12,7 @@ void UCEClonerCylinderLayout::SetBaseCount(int32 InBaseCount)
 	}
 
 	BaseCount = InBaseCount;
-	UpdateLayoutParameters();
+	MarkLayoutDirty();
 }
 
 void UCEClonerCylinderLayout::SetHeightCount(int32 InHeightCount)
@@ -23,7 +23,7 @@ void UCEClonerCylinderLayout::SetHeightCount(int32 InHeightCount)
 	}
 
 	HeightCount = InHeightCount;
-	UpdateLayoutParameters();
+	MarkLayoutDirty();
 }
 
 void UCEClonerCylinderLayout::SetHeight(float InHeight)
@@ -34,7 +34,7 @@ void UCEClonerCylinderLayout::SetHeight(float InHeight)
 	}
 
 	Height = InHeight;
-	UpdateLayoutParameters();
+	MarkLayoutDirty();
 }
 
 void UCEClonerCylinderLayout::SetRadius(float InRadius)
@@ -45,7 +45,7 @@ void UCEClonerCylinderLayout::SetRadius(float InRadius)
 	}
 
 	Radius = InRadius;
-	UpdateLayoutParameters();
+	MarkLayoutDirty();
 }
 
 void UCEClonerCylinderLayout::SetAngleStart(float InAngleStart)
@@ -56,7 +56,7 @@ void UCEClonerCylinderLayout::SetAngleStart(float InAngleStart)
 	}
 
 	AngleStart = InAngleStart;
-	UpdateLayoutParameters();
+	MarkLayoutDirty();
 }
 
 void UCEClonerCylinderLayout::SetAngleRatio(float InAngleRatio)
@@ -67,7 +67,7 @@ void UCEClonerCylinderLayout::SetAngleRatio(float InAngleRatio)
 	}
 
 	AngleRatio = InAngleRatio;
-	UpdateLayoutParameters();
+	MarkLayoutDirty();
 }
 
 void UCEClonerCylinderLayout::SetOrientMesh(bool bInOrientMesh)
@@ -78,7 +78,7 @@ void UCEClonerCylinderLayout::SetOrientMesh(bool bInOrientMesh)
 	}
 
 	bOrientMesh = bInOrientMesh;
-	UpdateLayoutParameters();
+	MarkLayoutDirty();
 }
 
 void UCEClonerCylinderLayout::SetPlane(ECEClonerPlane InPlane)
@@ -89,7 +89,7 @@ void UCEClonerCylinderLayout::SetPlane(ECEClonerPlane InPlane)
 	}
 
 	Plane = InPlane;
-	UpdateLayoutParameters();
+	MarkLayoutDirty();
 }
 
 void UCEClonerCylinderLayout::SetRotation(const FRotator& InRotation)
@@ -100,7 +100,7 @@ void UCEClonerCylinderLayout::SetRotation(const FRotator& InRotation)
 	}
 
 	Rotation = InRotation;
-	UpdateLayoutParameters();
+	MarkLayoutDirty();
 }
 
 void UCEClonerCylinderLayout::SetScale(const FVector& InScale)
@@ -111,7 +111,7 @@ void UCEClonerCylinderLayout::SetScale(const FVector& InScale)
 	}
 
 	Scale = InScale;
-	UpdateLayoutParameters();
+	MarkLayoutDirty();
 }
 
 #if WITH_EDITOR
diff --git a/Engine/Plugins/Experimental/ClonerEffector/Source/ClonerEffector/Private/Cloner/Layouts/CEClonerGridLayout.cpp b/Engine/Plugins/Experimental/ClonerEffector/Source/ClonerEffector/Private/Cloner/Layouts/CEClonerGridLayout.cpp
--- a/Engine/Plugins/Experimental/ClonerEffector/Source/ClonerEffector/Private/Cloner/Layouts/CEClonerGridLayout.cpp
+++ b/Engine/Plugins/Experimental/ClonerEffector/Source/ClonerEffector/Private/Cloner/Layouts/CEClonerGridLayout.cpp
@@ -14,7 +14,7 @@ void UCEClonerGridLayout::SetCountX(int32 InCountX)
 	}
 
 	CountX = InCountX;
-	UpdateLayoutParameters();
+	MarkLayoutDirty();
 }
 
 void UCEClonerGridLayout::SetCountY(int32 InCountY)
@@ -25,7 +25,7 @@ void UCEClonerGridLayout::SetCountY(int32 InCountY)
 	}
 
 	CountY = InCountY;
-	UpdateLayoutParameters();
+	MarkLayoutDirty();
 }
 
 void UCEClonerGridLayout::SetCountZ(int32 InCountZ)
@@ -36,7 +36,7 @@ void UCEClonerGridLayout::SetCountZ(int32 InCountZ)
 	}
 
 	CountZ = InCountZ;
-	UpdateLayoutParameters();
+	MarkLayoutDirty();
 }
 
 void UCEClonerGridLayout::SetSpacingX(float InSpacingX)
@@ -47,7 +47,7 @@ void UCEClonerGridLayout::SetSpacingX(float InSpacingX)
 	}
 
 	SpacingX = InSpacingX;
-	UpdateLayoutParameters();
+	MarkLayoutDirty();
 }
 
 void UCEClonerGridLayout::SetSpacingY(float InSpacingY)
@@ -58,7 +58,7 @@ void UCEClonerGridLayout::SetSpacingY(float InSpacingY)
 	}
 
 	SpacingY = InSpacingY;
-	UpdateLayoutParameters();
+	MarkLayoutDirty();
 }
 
 void UCEClonerGridLayout::SetSpacingZ(float InSpacingZ)
@@ -69,7 +69,7 @@ void UCEClonerGridLayout::SetSpacingZ(float InSpacingZ)
 	}
 
 	SpacingZ = InSpacingZ;
-	UpdateLayoutParameters();
+	MarkLayoutDirty();
 }
 
 void UCEClonerGridLayout::SetConstraint(ECEClonerGridConstraint InConstraint)
@@ -80,7 +80,7 @@ void UCEClonerGridLayout::SetConstraint(ECEClonerGridConstraint InConstraint)
 	}
 
 	Constraint = InConstraint;
-	UpdateLayoutParameters();
+	MarkLayoutDirty();
 }
 
 void UCEClonerGridLayout::SetInvertConstraint(bool bInInvertConstraint)
@@ -91,25 +91,25 @@ void UCEClonerGridLayout::SetInvertConstraint(bool bInInvertConstraint)
 	}
 
 	bInvertConstraint = bInInvertConstraint;
-	UpdateLayoutParameters();
+	MarkLayoutDirty();
 }
 
 void UCEClonerGridLayout::SetSphereConstraint(const FCEClonerGridConstraintSphere& InConstraint)
 {
 	SphereConstraint = InConstraint;
-	UpdateLayoutParameters();
+	MarkLayoutDirty();
 }
 
 void UCEClonerGridLayout::SetCylinderConstraint(const FCEClonerGridConstraintCylinder& InConstraint)
 {
 	CylinderConstraint = InConstraint;
-	UpdateLayoutParameters();
+	MarkLayoutDirty();
 }
 
 void UCEClonerGridLayout::SetTextureConstraint(const FCEClonerGridConstraintTexture& InConstraint)
 {
 	TextureConstraint = InConstraint;
-	UpdateLayoutParameters();
+	MarkLayoutDirty();
 }
 
 void UCEClonerGridLayout::SetTwistFactor(float InFactor)
@@ -120,18 +120,18 @@ void UCEClonerGridLayout::SetTwistFactor(float InFactor)
 	}
 
 	TwistFactor = InFactor;
-	UpdateLayoutParameters();
+	MarkLayoutDirty();
 }
 
-void UCEClonerGridLayout::SetTwistAxis(ENiagaraOrientationAxis InAxis)
+void UCEClonerGridLayout::SetTwistAxis(ECEClonerAxis InAxis)
 {
-	if (TwistAxis == InAxis)
+	if (TwistAxis == InAxis || InAxis == ECEClonerAxis::Custom)
 	{
 		return;
 	}
 
 	TwistAxis = InAxis;
-	UpdateLayoutParameters();
+	MarkLayoutDirty();
 }
 
 #if WITH_EDITOR
diff --git a/Engine/Plugins/Experimental/ClonerEffector/Source/ClonerEffector/Private/Cloner/Layouts/CEClonerHoneycombLayout.cpp b/Engine/Plugins/Experimental/ClonerEffector/Source/ClonerEffector/Private/Cloner/Layouts/CEClonerHoneycombLayout.cpp
--- a/Engine/Plugins/Experimental/ClonerEffector/Source/ClonerEffector/Private/Cloner/Layouts/CEClonerHoneycombLayout.cpp
+++ b/Engine/Plugins/Experimental/ClonerEffector/Source/ClonerEffector/Private/Cloner/Layouts/CEClonerHoneycombLayout.cpp
@@ -29,7 +29,7 @@ void UCEClonerHoneycombLayout::SetWidthCount(int32 InWidthCount)
 	}
 
 	WidthCount = InWidthCount;
-	UpdateLayoutParameters();
+	MarkLayoutDirty();
 }
 
 void UCEClonerHoneycombLayout::SetHeightCount(int32 InHeightCount)
@@ -40,7 +40,7 @@ void UCEClonerHoneycombLayout::SetHeightCount(int32 InHeightCount)
 	}
 
 	HeightCount = InHeightCount;
-	UpdateLayoutParameters();
+	MarkLayoutDirty();
 }
 
 void UCEClonerHoneycombLayout::SetWidthOffset(float InWidthOffset)
@@ -51,7 +51,7 @@ void UCEClonerHoneycombLayout::SetWidthOffset(float InWidthOffset)
 	}
 
 	WidthOffset = InWidthOffset;
-	UpdateLayoutParameters();
+	MarkLayoutDirty();
 }
 
 void UCEClonerHoneycombLayout::SetHeightOffset(float InHeightOffset)
@@ -62,7 +62,7 @@ void UCEClonerHoneycombLayout::SetHeightOffset(float InHeightOffset)
 	}
 
 	HeightOffset = InHeightOffset;
-	UpdateLayoutParameters();
+	MarkLayoutDirty();
 }
 
 void UCEClonerHoneycombLayout::SetHeightSpacing(float InHeightSpacing)
@@ -73,7 +73,7 @@ void UCEClonerHoneycombLayout::SetHeightSpacing(float InHeightSpacing)
 	}
 
 	HeightSpacing = InHeightSpacing;
-	UpdateLayoutParameters();
+	MarkLayoutDirty();
 }
 
 void UCEClonerHoneycombLayout::SetWidthSpacing(float InWidthSpacing)
@@ -84,7 +84,7 @@ void UCEClonerHoneycombLayout::SetWidthSpacing(float InWidthSpacing)
 	}
 
 	WidthSpacing = InWidthSpacing;
-	UpdateLayoutParameters();
+	MarkLayoutDirty();
 }
 
 void UCEClonerHoneycombLayout::SetTwistFactor(float InFactor)
@@ -95,12 +95,12 @@ void UCEClonerHoneycombLayout::SetTwistFactor(float InFactor)
 	}
 
 	TwistFactor = InFactor;
-	UpdateLayoutParameters();
+	MarkLayoutDirty();
 }
 
-void UCEClonerHoneycombLayout::SetTwistAxis(ENiagaraOrientationAxis InAxis)
+void UCEClonerHoneycombLayout::SetTwistAxis(ECEClonerAxis InAxis)
 {
-	if (TwistAxis == InAxis)
+	if (TwistAxis == InAxis || InAxis == ECEClonerAxis::Custom)
 	{
 		return;
 	}
@@ -136,27 +136,27 @@ void UCEClonerHoneycombLayout::OnTwistAxisChanged()
 	// Restrict twist axis to plane axis
 	if (Plane == ECEClonerPlane::XY)
 	{
-		if (TwistAxis == ENiagaraOrientationAxis::ZAxis)
+		if (TwistAxis == ECEClonerAxis::Z)
 		{
-			TwistAxis = ENiagaraOrientationAxis::XAxis;
+			TwistAxis = ECEClonerAxis::X;
 		}
 	}
 	else if (Plane == ECEClonerPlane::XZ)
 	{
-		if (TwistAxis == ENiagaraOrientationAxis::YAxis)
+		if (TwistAxis == ECEClonerAxis::Y)
 		{
-			TwistAxis = ENiagaraOrientationAxis::XAxis;
+			TwistAxis = ECEClonerAxis::X;
 		}
 	}
 	else if (Plane == ECEClonerPlane::YZ)
 	{
-		if (TwistAxis == ENiagaraOrientationAxis::XAxis)
+		if (TwistAxis == ECEClonerAxis::X)
 		{
-			TwistAxis = ENiagaraOrientationAxis::YAxis;
+			TwistAxis = ECEClonerAxis::Y;
 		}
 	}
 
-	UpdateLayoutParameters();
+	MarkLayoutDirty();
 }
 
 void UCEClonerHoneycombLayout::OnLayoutParametersChanged(UCEClonerComponent* InComponent)
diff --git a/Engine/Plugins/Experimental/ClonerEffector/Source/ClonerEffector/Private/Cloner/Layouts/CEClonerLayoutBase.cpp b/Engine/Plugins/Experimental/ClonerEffector/Source/ClonerEffector/Private/Cloner/Layouts/CEClonerLayoutBase.cpp
--- a/Engine/Plugins/Experimental/ClonerEffector/Source/ClonerEffector/Private/Cloner/Layouts/CEClonerLayoutBase.cpp
+++ b/Engine/Plugins/Experimental/ClonerEffector/Source/ClonerEffector/Private/Cloner/Layouts/CEClonerLayoutBase.cpp
@@ -234,9 +234,30 @@ TSet<FName> UCEClonerLayoutBase::GetSupportedExtensions() const
 	return ModuleSupported;
 }
 
+bool UCEClonerLayoutBase::IsLayoutDirty() const
+{
+	return EnumHasAnyFlags(LayoutStatus, ECEClonerSystemStatus::ParametersDirty);
+}
+
+void UCEClonerLayoutBase::PostEditImport()
+{
+	Super::PostEditImport();
+
+	MarkLayoutDirty();
+}
+
+#if WITH_EDITOR
+void UCEClonerLayoutBase::PostEditUndo()
+{
+	Super::PostEditUndo();
+
+	MarkLayoutDirty();
+}
+#endif
+
 void UCEClonerLayoutBase::OnLayoutPropertyChanged()
 {
-	UpdateLayoutParameters();
+	MarkLayoutDirty();
 }
 
 void UCEClonerLayoutBase::OnSystemPackageLoaded(const FName& InName, UPackage* InPackage, EAsyncLoadingResult::Type InResult)
@@ -300,7 +321,7 @@ AActor* UCEClonerLayoutBase::GetClonerActor() const
 	return nullptr;
 }
 
-void UCEClonerLayoutBase::UpdateLayoutParameters(bool bInUpdateCloner, bool bInImmediate)
+void UCEClonerLayoutBase::UpdateLayoutParameters()
 {
 	if (!IsLayoutActive())
 	{
@@ -309,19 +330,28 @@ void UCEClonerLayoutBase::UpdateLayoutParameters(bool bInUpdateCloner, bool bInI
 
 	if (UCEClonerComponent* ClonerComponent = GetClonerComponent())
 	{
+		if (!ClonerComponent->GetEnabled())
+		{
+			return;
+		}
+
 		OnLayoutParametersChanged(ClonerComponent);
 
-		if (bInUpdateCloner)
+		if (EnumHasAnyFlags(LayoutStatus, ECEClonerSystemStatus::SimulationDirty))
 		{
-			RequestClonerUpdate(bInImmediate);
+			ClonerComponent->RequestClonerUpdate(/*Immediate*/false);
 		}
 	}
+
+	LayoutStatus = ECEClonerSystemStatus::UpToDate;
 }
 
-void UCEClonerLayoutBase::RequestClonerUpdate(bool bInImmediate) const
+void UCEClonerLayoutBase::MarkLayoutDirty(bool bInUpdateCloner)
 {
-	if (UCEClonerComponent* ClonerComponent = GetClonerComponent())
+	EnumAddFlags(LayoutStatus, ECEClonerSystemStatus::ParametersDirty);
+
+	if (bInUpdateCloner)
 	{
-		ClonerComponent->RequestClonerUpdate(bInImmediate);
+		EnumAddFlags(LayoutStatus, ECEClonerSystemStatus::SimulationDirty);
 	}
 }
diff --git a/Engine/Plugins/Experimental/ClonerEffector/Source/ClonerEffector/Private/Cloner/Layouts/CEClonerLineLayout.cpp b/Engine/Plugins/Experimental/ClonerEffector/Source/ClonerEffector/Private/Cloner/Layouts/CEClonerLineLayout.cpp
--- a/Engine/Plugins/Experimental/ClonerEffector/Source/ClonerEffector/Private/Cloner/Layouts/CEClonerLineLayout.cpp
+++ b/Engine/Plugins/Experimental/ClonerEffector/Source/ClonerEffector/Private/Cloner/Layouts/CEClonerLineLayout.cpp
@@ -12,7 +12,7 @@ void UCEClonerLineLayout::SetCount(int32 InCount)
 	}
 
 	Count = InCount;
-	UpdateLayoutParameters();
+	MarkLayoutDirty();
 }
 
 void UCEClonerLineLayout::SetSpacing(float InSpacing)
@@ -23,7 +23,7 @@ void UCEClonerLineLayout::SetSpacing(float InSpacing)
 	}
 
 	Spacing = InSpacing;
-	UpdateLayoutParameters();
+	MarkLayoutDirty();
 }
 
 void UCEClonerLineLayout::SetAxis(ECEClonerAxis InAxis)
@@ -34,7 +34,7 @@ void UCEClonerLineLayout::SetAxis(ECEClonerAxis InAxis)
 	}
 
 	Axis = InAxis;
-	UpdateLayoutParameters();
+	MarkLayoutDirty();
 }
 
 void UCEClonerLineLayout::SetDirection(const FVector& InDirection)
@@ -45,7 +45,7 @@ void UCEClonerLineLayout::SetDirection(const FVector& InDirection)
 	}
 
 	Direction = InDirection;
-	UpdateLayoutParameters();
+	MarkLayoutDirty();
 }
 
 void UCEClonerLineLayout::SetRotation(const FRotator& InRotation)
@@ -56,7 +56,7 @@ void UCEClonerLineLayout::SetRotation(const FRotator& InRotation)
 	}
 
 	Rotation = InRotation;
-	UpdateLayoutParameters();
+	MarkLayoutDirty();
 }
 
 #if WITH_EDITOR
diff --git a/Engine/Plugins/Experimental/ClonerEffector/Source/ClonerEffector/Private/Cloner/Layouts/CEClonerMeshLayout.cpp b/Engine/Plugins/Experimental/ClonerEffector/Source/ClonerEffector/Private/Cloner/Layouts/CEClonerMeshLayout.cpp
--- a/Engine/Plugins/Experimental/ClonerEffector/Source/ClonerEffector/Private/Cloner/Layouts/CEClonerMeshLayout.cpp
+++ b/Engine/Plugins/Experimental/ClonerEffector/Source/ClonerEffector/Private/Cloner/Layouts/CEClonerMeshLayout.cpp
@@ -29,7 +29,7 @@ void UCEClonerMeshLayout::SetCount(int32 InCount)
 	}
 
 	Count = InCount;
-	UpdateLayoutParameters();
+	MarkLayoutDirty();
 }
 
 void UCEClonerMeshLayout::SetAsset(ECEClonerMeshAsset InAsset)
@@ -40,7 +40,7 @@ void UCEClonerMeshLayout::SetAsset(ECEClonerMeshAsset InAsset)
 	}
 
 	Asset = InAsset;
-	UpdateLayoutParameters();
+	MarkLayoutDirty();
 }
 
 void UCEClonerMeshLayout::SetSampleData(ECEClonerMeshSampleData InSampleData)
@@ -51,7 +51,7 @@ void UCEClonerMeshLayout::SetSampleData(ECEClonerMeshSampleData InSampleData)
 	}
 
 	SampleData = InSampleData;
-	UpdateLayoutParameters();
+	MarkLayoutDirty();
 }
 
 void UCEClonerMeshLayout::SetSampleActorWeak(const TWeakObjectPtr<AActor>& InSampleActor)
@@ -62,7 +62,7 @@ void UCEClonerMeshLayout::SetSampleActorWeak(const TWeakObjectPtr<AActor>& InSam
 	}
 
 	SampleActorWeak = InSampleActor;
-	UpdateLayoutParameters();
+	MarkLayoutDirty();
 }
 
 void UCEClonerMeshLayout::SetSampleActor(AActor* InActor)
@@ -165,5 +165,5 @@ void UCEClonerMeshLayout::OnLayoutParametersChanged(UCEClonerComponent* InCompon
 
 void UCEClonerMeshLayout::OnSampleMeshTransformed(USceneComponent* InComponent, EUpdateTransformFlags InFlags, ETeleportType InType)
 {
-	UpdateLayoutParameters();
+	MarkLayoutDirty();
 }
diff --git a/Engine/Plugins/Experimental/ClonerEffector/Source/ClonerEffector/Private/Cloner/Layouts/CEClonerSphereRandomLayout.cpp b/Engine/Plugins/Experimental/ClonerEffector/Source/ClonerEffector/Private/Cloner/Layouts/CEClonerSphereRandomLayout.cpp
--- a/Engine/Plugins/Experimental/ClonerEffector/Source/ClonerEffector/Private/Cloner/Layouts/CEClonerSphereRandomLayout.cpp
+++ b/Engine/Plugins/Experimental/ClonerEffector/Source/ClonerEffector/Private/Cloner/Layouts/CEClonerSphereRandomLayout.cpp
@@ -12,7 +12,7 @@ void UCEClonerSphereRandomLayout::SetCount(int32 InCount)
 	}
 
 	Count = InCount;
-	UpdateLayoutParameters();
+	MarkLayoutDirty();
 }
 
 void UCEClonerSphereRandomLayout::SetRadius(float InRadius)
@@ -23,7 +23,7 @@ void UCEClonerSphereRandomLayout::SetRadius(float InRadius)
 	}
 
 	Radius = InRadius;
-	UpdateLayoutParameters();
+	MarkLayoutDirty();
 }
 
 void UCEClonerSphereRandomLayout::SetDistribution(float InDistribution)
@@ -39,7 +39,7 @@ void UCEClonerSphereRandomLayout::SetDistribution(float InDistribution)
 	}
 
 	Distribution = InDistribution;
-	UpdateLayoutParameters();
+	MarkLayoutDirty();
 }
 
 void UCEClonerSphereRandomLayout::SetLongitude(float InLongitude)
@@ -55,7 +55,7 @@ void UCEClonerSphereRandomLayout::SetLongitude(float InLongitude)
 	}
 
 	Longitude = InLongitude;
-	UpdateLayoutParameters();
+	MarkLayoutDirty();
 }
 
 void UCEClonerSphereRandomLayout::SetLatitude(float InLatitude)
@@ -71,7 +71,7 @@ void UCEClonerSphereRandomLayout::SetLatitude(float InLatitude)
 	}
 
 	Latitude = InLatitude;
-	UpdateLayoutParameters();
+	MarkLayoutDirty();
 }
 
 void UCEClonerSphereRandomLayout::SetOrientMesh(bool bInOrientMesh)
@@ -82,7 +82,7 @@ void UCEClonerSphereRandomLayout::SetOrientMesh(bool bInOrientMesh)
 	}
 
 	bOrientMesh = bInOrientMesh;
-	UpdateLayoutParameters();
+	MarkLayoutDirty();
 }
 
 void UCEClonerSphereRandomLayout::SetRotation(const FRotator& InRotation)
@@ -93,7 +93,7 @@ void UCEClonerSphereRandomLayout::SetRotation(const FRotator& InRotation)
 	}
 
 	Rotation = InRotation;
-	UpdateLayoutParameters();
+	MarkLayoutDirty();
 }
 
 void UCEClonerSphereRandomLayout::SetScale(const FVector& InScale)
@@ -104,7 +104,7 @@ void UCEClonerSphereRandomLayout::SetScale(const FVector& InScale)
 	}
 
 	Scale = InScale;
-	UpdateLayoutParameters();
+	MarkLayoutDirty();
 }
 
 #if WITH_EDITOR
diff --git a/Engine/Plugins/Experimental/ClonerEffector/Source/ClonerEffector/Private/Cloner/Layouts/CEClonerSphereUniformLayout.cpp b/Engine/Plugins/Experimental/ClonerEffector/Source/ClonerEffector/Private/Cloner/Layouts/CEClonerSphereUniformLayout.cpp
--- a/Engine/Plugins/Experimental/ClonerEffector/Source/ClonerEffector/Private/Cloner/Layouts/CEClonerSphereUniformLayout.cpp
+++ b/Engine/Plugins/Experimental/ClonerEffector/Source/ClonerEffector/Private/Cloner/Layouts/CEClonerSphereUniformLayout.cpp
@@ -12,7 +12,7 @@ void UCEClonerSphereUniformLayout::SetCount(int32 InCount)
 	}
 
 	Count = InCount;
-	UpdateLayoutParameters();
+	MarkLayoutDirty();
 }
 
 void UCEClonerSphereUniformLayout::SetRadius(float InRadius)
@@ -23,7 +23,7 @@ void UCEClonerSphereUniformLayout::SetRadius(float InRadius)
 	}
 
 	Radius = InRadius;
-	UpdateLayoutParameters();
+	MarkLayoutDirty();
 }
 
 void UCEClonerSphereUniformLayout::SetRatio(float InRatio)
@@ -34,7 +34,7 @@ void UCEClonerSphereUniformLayout::SetRatio(float InRatio)
 	}
 
 	Ratio = InRatio;
-	UpdateLayoutParameters();
+	MarkLayoutDirty();
 }
 
 void UCEClonerSphereUniformLayout::SetOrientMesh(bool bInOrientMesh)
@@ -45,7 +45,7 @@ void UCEClonerSphereUniformLayout::SetOrientMesh(bool bInOrientMesh)
 	}
 
 	bOrientMesh = bInOrientMesh;
-	UpdateLayoutParameters();
+	MarkLayoutDirty();
 }
 
 void UCEClonerSphereUniformLayout::SetRotation(const FRotator& InRotation)
@@ -56,7 +56,7 @@ void UCEClonerSphereUniformLayout::SetRotation(const FRotator& InRotation)
 	}
 
 	Rotation = InRotation;
-	UpdateLayoutParameters();
+	MarkLayoutDirty();
 }
 
 void UCEClonerSphereUniformLayout::SetScale(const FVector& InScale)
@@ -67,7 +67,7 @@ void UCEClonerSphereUniformLayout::SetScale(const FVector& InScale)
 	}
 
 	Scale = InScale;
-	UpdateLayoutParameters();
+	MarkLayoutDirty();
 }
 
 #if WITH_EDITOR
diff --git a/Engine/Plugins/Experimental/ClonerEffector/Source/ClonerEffector/Private/Cloner/Layouts/CEClonerSplineLayout.cpp b/Engine/Plugins/Experimental/ClonerEffector/Source/ClonerEffector/Private/Cloner/Layouts/CEClonerSplineLayout.cpp
--- a/Engine/Plugins/Experimental/ClonerEffector/Source/ClonerEffector/Private/Cloner/Layouts/CEClonerSplineLayout.cpp
+++ b/Engine/Plugins/Experimental/ClonerEffector/Source/ClonerEffector/Private/Cloner/Layouts/CEClonerSplineLayout.cpp
@@ -27,7 +27,7 @@ void UCEClonerSplineLayout::SetCount(int32 InCount)
 	}
 
 	Count = InCount;
-	UpdateLayoutParameters();
+	MarkLayoutDirty();
 }
 
 void UCEClonerSplineLayout::SetSplineActorWeak(const TWeakObjectPtr<AActor>& InSplineActor)
@@ -38,7 +38,7 @@ void UCEClonerSplineLayout::SetSplineActorWeak(const TWeakObjectPtr<AActor>& InS
 	}
 
 	SplineActorWeak = InSplineActor;
-	UpdateLayoutParameters();
+	MarkLayoutDirty();
 }
 
 void UCEClonerSplineLayout::SetSplineActor(AActor* InSplineActor)
@@ -54,7 +54,7 @@ void UCEClonerSplineLayout::SetOrientMesh(bool bInOrientMesh)
 	}
 
 	bOrientMesh = bInOrientMesh;
-	UpdateLayoutParameters();
+	MarkLayoutDirty();
 }
 
 #if WITH_EDITOR
@@ -152,13 +152,14 @@ void UCEClonerSplineLayout::OnLayoutParametersChanged(UCEClonerComponent* InComp
 	if (USplineComponent* SplineComponent = SplineComponentWeak.Get())
     {
 		SplineComponent->TransformUpdated.RemoveAll(this);
+
 		SplineComponentWeak.Reset();
     }
 
 	SplineComponentWeak = nullptr;
 
 	// bind
-	if (AActor* SplineActor = SplineActorWeak.Get())
+	if (const AActor* SplineActor = SplineActorWeak.Get())
 	{
 		if (USplineComponent* SplineComponent = SplineActor->FindComponentByClass<USplineComponent>())
 		{
@@ -181,13 +182,13 @@ void UCEClonerSplineLayout::OnLayoutParametersChanged(UCEClonerComponent* InComp
 
 void UCEClonerSplineLayout::OnSampleSplineTransformed(USceneComponent* InComponent, EUpdateTransformFlags InFlags, ETeleportType InType)
 {
-	UpdateLayoutParameters();
+	MarkLayoutDirty();
 }
 
 void UCEClonerSplineLayout::OnSampleSplineRenderStateUpdated(UActorComponent& InComponent)
 {
 	if (SplineActorWeak.IsValid() && InComponent.GetOwner() == SplineActorWeak.Get())
 	{
-		UpdateLayoutParameters();
+		MarkLayoutDirty();
 	}
 }
diff --git a/Engine/Plugins/Experimental/ClonerEffector/Source/ClonerEffector/Private/Effector/CEEffectorComponent.cpp b/Engine/Plugins/Experimental/ClonerEffector/Source/ClonerEffector/Private/Effector/CEEffectorComponent.cpp
--- a/Engine/Plugins/Experimental/ClonerEffector/Source/ClonerEffector/Private/Effector/CEEffectorComponent.cpp
+++ b/Engine/Plugins/Experimental/ClonerEffector/Source/ClonerEffector/Private/Effector/CEEffectorComponent.cpp
@@ -19,6 +19,18 @@
 
 DEFINE_LOG_CATEGORY_STATIC(LogCEEffectorComponent, Log, All);
 
+#if WITH_EDITOR
+FName UCEEffectorComponent::GetModeNamePropertyName()
+{
+	return GET_MEMBER_NAME_CHECKED(UCEEffectorComponent, ModeName);
+}
+
+FName UCEEffectorComponent::GetTypeNamePropertyName()
+{
+	return GET_MEMBER_NAME_CHECKED(UCEEffectorComponent, TypeName);
+}
+#endif
+
 UCEEffectorComponent::UCEEffectorComponent()
 {
 	PrimaryComponentTick.bCanEverTick = false;
diff --git a/Engine/Plugins/Experimental/ClonerEffector/Source/ClonerEffector/Private/Effector/CEEffectorExtensionBase.cpp b/Engine/Plugins/Experimental/ClonerEffector/Source/ClonerEffector/Private/Effector/CEEffectorExtensionBase.cpp
--- a/Engine/Plugins/Experimental/ClonerEffector/Source/ClonerEffector/Private/Effector/CEEffectorExtensionBase.cpp
+++ b/Engine/Plugins/Experimental/ClonerEffector/Source/ClonerEffector/Private/Effector/CEEffectorExtensionBase.cpp
@@ -4,12 +4,19 @@
 
 #include "Effector/CEEffectorComponent.h"
 
+#if WITH_EDITOR
+FCEExtensionSection UCEEffectorExtensionBase::GetExtensionSection() const
+{
+	return UE::ClonerEffector::EditorSection::GetExtensionSectionFromClass(GetClass());
+}
+#endif
+
 UCEEffectorComponent* UCEEffectorExtensionBase::GetEffectorComponent() const
 {
 	return GetTypedOuter<UCEEffectorComponent>();
 }
 
-void UCEEffectorExtensionBase::UpdateExtensionParameters(bool bInUpdateLinkedCloners, bool bInImmediate)
+void UCEEffectorExtensionBase::UpdateExtensionParameters(bool bInUpdateLinkedCloners)
 {
 	if (!IsExtensionActive())
 	{
@@ -27,7 +34,7 @@ void UCEEffectorExtensionBase::UpdateExtensionParameters(bool bInUpdateLinkedClo
 
 		if (bInUpdateLinkedCloners)
 		{
-			EffectorComponent->RequestClonerUpdate(bInImmediate);
+			EffectorComponent->RequestClonerUpdate(/** Immediate */false);
 		}
 	}
 }
diff --git a/Engine/Plugins/Experimental/ClonerEffector/Source/ClonerEffector/Public/CEClonerEffectorShared.h b/Engine/Plugins/Experimental/ClonerEffector/Source/ClonerEffector/Public/CEClonerEffectorShared.h
--- a/Engine/Plugins/Experimental/ClonerEffector/Source/ClonerEffector/Public/CEClonerEffectorShared.h
+++ b/Engine/Plugins/Experimental/ClonerEffector/Source/ClonerEffector/Public/CEClonerEffectorShared.h
@@ -12,6 +12,7 @@ class AActor;
 class ADynamicMeshActor;
 class AStaticMeshActor;
 class UCEClonerComponent;
+class UClass;
 class UNiagaraDataChannelReader;
 class UNiagaraDataChannelWriter;
 class UDynamicMesh;
@@ -262,6 +263,17 @@ enum class ECEClonerMeshConversion : uint8
 	InstancedStaticMesh
 };
 
+/** Enumerates all states for extension/layout */
+enum class ECEClonerSystemStatus : uint8
+{
+	/** Nothing to do */
+	UpToDate = 0,
+	/** Parameters needs an update */
+	ParametersDirty = 1 << 0,
+	/** Simulation needs an update */
+	SimulationDirty = 1 << 1
+};
+
 USTRUCT()
 struct FCEClonerAttachmentItem
 {
@@ -453,6 +465,7 @@ struct FCEClonerEffectorDataInterfaces
 	void Remove(int32 InIndex) const;
 	bool IsValid() const;
 	int32 Num() const;
+	void Commit() const;
 
 	UNiagaraDataInterfaceArrayInt32* GetIndexArray() const;
 
@@ -531,21 +544,10 @@ struct FCEExtensionSection
 namespace UE::ClonerEffector
 {
 #if WITH_EDITOR
-	namespace ClonerSection
-	{
-		const FCEExtensionSection ClonerSection(TEXT("Cloner"), 0);
-		const FCEExtensionSection EffectorSection(TEXT("Effector"), 1);
-		const FCEExtensionSection EmissionSection(TEXT("Emission"), 2);
-		const FCEExtensionSection PhysicsSection(TEXT("Physics"), 3);
-		const FCEExtensionSection RenderingSection(TEXT("Rendering"), 4);
-	}
-
-	namespace EffectorSection
+	namespace EditorSection
 	{
-		const FCEExtensionSection EffectorSection(TEXT("Effector"), 0);
-		const FCEExtensionSection ModeSection(TEXT("Mode"), 1);
-		const FCEExtensionSection ShapeSection(TEXT("Shape"), 2);
-		const FCEExtensionSection ForcesSection(TEXT("Forces"), 3);
+		/** Retrieves section metadata from class */
+		FCEExtensionSection GetExtensionSectionFromClass(UClass* InClass);
 	}
 #endif
 
diff --git a/Engine/Plugins/Experimental/ClonerEffector/Source/ClonerEffector/Public/Cloner/CEClonerComponent.h b/Engine/Plugins/Experimental/ClonerEffector/Source/ClonerEffector/Public/Cloner/CEClonerComponent.h
--- a/Engine/Plugins/Experimental/ClonerEffector/Source/ClonerEffector/Public/Cloner/CEClonerComponent.h
+++ b/Engine/Plugins/Experimental/ClonerEffector/Source/ClonerEffector/Public/Cloner/CEClonerComponent.h
@@ -37,22 +37,24 @@ public:
 	/** Show material warning notification when missing niagara usage flag */
 	static void ShowMaterialWarning(int32 InMaterialCount);
 
-	static CLONEREFFECTOR_API FName GetActiveExtensionsName();
+	static CLONEREFFECTOR_API FName GetActiveExtensionsPropertyName();
 
-	static CLONEREFFECTOR_API FName GetActiveLayoutName();
+	static CLONEREFFECTOR_API FName GetActiveLayoutPropertyName();
+
+	static CLONEREFFECTOR_API FName GetLayoutNamePropertyName();
 #endif
 
-	TMulticastDelegateRegistration<void(UCEClonerComponent*)>& OnClonerMeshUpdated()
+	static TMulticastDelegateRegistration<void(UCEClonerComponent*)>& OnClonerMeshUpdated()
 	{
 		return OnClonerMeshUpdatedDelegate;
 	}
 
-	TMulticastDelegateRegistration<void(UCEClonerComponent*, UCEClonerLayoutBase*)>& OnClonerLayoutLoaded()
+	static TMulticastDelegateRegistration<void(UCEClonerComponent*, UCEClonerLayoutBase*)>& OnClonerLayoutLoaded()
 	{
 		return OnClonerLayoutLoadedDelegate;
 	}
 
-	TMulticastDelegateRegistration<void(UCEClonerComponent*)>& OnClonerInitialized()
+	static TMulticastDelegateRegistration<void(UCEClonerComponent*)>& OnClonerInitialized()
 	{
 		return OnClonerInitializedDelegate;
 	}
@@ -210,15 +212,15 @@ public:
 protected:
 	/** Called when meshes have been updated */
 	DECLARE_MULTICAST_DELEGATE_OneParam(FOnClonerMeshUpdated, UCEClonerComponent* /** ClonerComponent */)
-	FOnClonerMeshUpdated OnClonerMeshUpdatedDelegate;
+	CLONEREFFECTOR_API static FOnClonerMeshUpdated OnClonerMeshUpdatedDelegate;
 
 	/** Called when new cloner layout is loaded */
 	DECLARE_MULTICAST_DELEGATE_TwoParams(FOnClonerLayoutLoaded, UCEClonerComponent* /** ClonerComponent */, UCEClonerLayoutBase* /** InLayout */)
-	FOnClonerLayoutLoaded OnClonerLayoutLoadedDelegate;
+	CLONEREFFECTOR_API static FOnClonerLayoutLoaded OnClonerLayoutLoadedDelegate;
 
 	/** Called when cloner is initialized */
 	DECLARE_MULTICAST_DELEGATE_OneParam(FOnClonerInitialized, UCEClonerComponent* /** ClonerComponent */)
-	FOnClonerInitialized OnClonerInitializedDelegate;
+	CLONEREFFECTOR_API static FOnClonerInitialized OnClonerInitializedDelegate;
 
 	//~ Begin UObject
 	virtual void PostInitProperties() override;
@@ -312,11 +314,11 @@ protected:
 	FName LayoutName = NAME_None;
 
 	/** Active layout used */
-	UPROPERTY(VisibleInstanceOnly, BlueprintReadOnly, Instanced, Transient, DuplicateTransient, NonTransactional, Category="Layout", meta=(DisplayAfter="LayoutName"))
+	UPROPERTY(Transient, DuplicateTransient, NonTransactional)
 	TObjectPtr<UCEClonerLayoutBase> ActiveLayout;
 
 	/** Active Extensions on this layout */
-	UPROPERTY(VisibleInstanceOnly, BlueprintReadOnly, Instanced, Transient, DuplicateTransient, NonTransactional, Category="Layout")
+	UPROPERTY(Transient, DuplicateTransient, NonTransactional)
 	TArray<TObjectPtr<UCEClonerExtensionBase>> ActiveExtensions;
 
 	/** Layout instances cached */
@@ -350,9 +352,6 @@ private:
 	/** Called to trigger an update of cloner rendering state tree */
 	void UpdateClonerRenderState();
 
-	/** Forces a refresh of the active system parameters in niagara store */
-	void RefreshUserParameters() const;
-
 	void OnEnabledChanged();
 	void OnClonerEnabled();
 	void OnClonerDisabled();
diff --git a/Engine/Plugins/Experimental/ClonerEffector/Source/ClonerEffector/Public/Cloner/Extensions/CEClonerCollisionExtension.h b/Engine/Plugins/Experimental/ClonerEffector/Source/ClonerEffector/Public/Cloner/Extensions/CEClonerCollisionExtension.h
--- a/Engine/Plugins/Experimental/ClonerEffector/Source/ClonerEffector/Public/Cloner/Extensions/CEClonerCollisionExtension.h
+++ b/Engine/Plugins/Experimental/ClonerEffector/Source/ClonerEffector/Public/Cloner/Extensions/CEClonerCollisionExtension.h
@@ -8,7 +8,7 @@
 #include "CEClonerCollisionExtension.generated.h"
 
 /** Extension dealing with collisions and physics related options */
-UCLASS(MinimalAPI, BlueprintType, Within=CEClonerComponent, AutoExpandCategories=("Collisions"))
+UCLASS(MinimalAPI, BlueprintType, Within=CEClonerComponent, AutoExpandCategories=("Collisions"), meta=(Section="Physics", Priority=3))
 class UCEClonerCollisionExtension : public UCEClonerExtensionBase
 {
 	GENERATED_BODY()
@@ -18,9 +18,6 @@ public:
 		: UCEClonerExtensionBase(
 			TEXT("Collisions")
 			, 0
-#if WITH_EDITOR
-			, UE::ClonerEffector::ClonerSection::PhysicsSection
-#endif
 		)
 	{}
 
diff --git a/Engine/Plugins/Experimental/ClonerEffector/Source/ClonerEffector/Public/Cloner/Extensions/CEClonerEffectorExtension.h b/Engine/Plugins/Experimental/ClonerEffector/Source/ClonerEffector/Public/Cloner/Extensions/CEClonerEffectorExtension.h
--- a/Engine/Plugins/Experimental/ClonerEffector/Source/ClonerEffector/Public/Cloner/Extensions/CEClonerEffectorExtension.h
+++ b/Engine/Plugins/Experimental/ClonerEffector/Source/ClonerEffector/Public/Cloner/Extensions/CEClonerEffectorExtension.h
@@ -12,7 +12,7 @@ class UCEEffectorComponent;
 struct FCEClonerEffectorDataInterfaces;
 
 /** Extension dealing with effectors options */
-UCLASS(MinimalAPI, BlueprintType, Within=CEClonerComponent)
+UCLASS(MinimalAPI, BlueprintType, Within=CEClonerComponent, meta=(Section="Effector", Priority=1))
 class UCEClonerEffectorExtension : public UCEClonerExtensionBase
 {
 	GENERATED_BODY()
@@ -26,9 +26,6 @@ public:
 		: UCEClonerExtensionBase(
 			TEXT("Effector")
 			, 0
-#if WITH_EDITOR
-			, UE::ClonerEffector::ClonerSection::EffectorSection
-#endif
 		)
 	{}
 
diff --git a/Engine/Plugins/Experimental/ClonerEffector/Source/ClonerEffector/Public/Cloner/Extensions/CEClonerEmitterSpawnExtension.h b/Engine/Plugins/Experimental/ClonerEffector/Source/ClonerEffector/Public/Cloner/Extensions/CEClonerEmitterSpawnExtension.h
--- a/Engine/Plugins/Experimental/ClonerEffector/Source/ClonerEffector/Public/Cloner/Extensions/CEClonerEmitterSpawnExtension.h
+++ b/Engine/Plugins/Experimental/ClonerEffector/Source/ClonerEffector/Public/Cloner/Extensions/CEClonerEmitterSpawnExtension.h
@@ -10,7 +10,7 @@
 class UNiagaraDataInterfaceCurve;
 
 /** Extension dealing with clones spawning options */
-UCLASS(MinimalAPI, BlueprintType, Within=CEClonerComponent)
+UCLASS(MinimalAPI, BlueprintType, Within=CEClonerComponent, meta=(Section="Emission", Priority=2))
 class UCEClonerEmitterSpawnExtension : public UCEClonerExtensionBase
 {
 	GENERATED_BODY()
@@ -20,9 +20,6 @@ public:
 		: UCEClonerExtensionBase(
 			TEXT("Spawn")
 			, 0
-#if WITH_EDITOR
-			, UE::ClonerEffector::ClonerSection::EmissionSection
-#endif
 		)
 	{}
 
diff --git a/Engine/Plugins/Experimental/ClonerEffector/Source/ClonerEffector/Public/Cloner/Extensions/CEClonerExtensionBase.h b/Engine/Plugins/Experimental/ClonerEffector/Source/ClonerEffector/Public/Cloner/Extensions/CEClonerExtensionBase.h
--- a/Engine/Plugins/Experimental/ClonerEffector/Source/ClonerEffector/Public/Cloner/Extensions/CEClonerExtensionBase.h
+++ b/Engine/Plugins/Experimental/ClonerEffector/Source/ClonerEffector/Public/Cloner/Extensions/CEClonerExtensionBase.h
@@ -20,24 +20,14 @@ public:
 		: UCEClonerExtensionBase(
 			NAME_None
 			, 0
-#if WITH_EDITOR
-			, FCEExtensionSection(NAME_None, INDEX_NONE)
-#endif
 		)
 	{}
 
 	UCEClonerExtensionBase(
 		FName InExtensionName
-		, int32 InExtensionPriority
-#if WITH_EDITOR
-		, const FCEExtensionSection& InExtensionSection
-#endif
-		)
+		, int32 InExtensionPriority)
 		: ExtensionName(InExtensionName)
 		, ExtensionPriority(InExtensionPriority)
-#if WITH_EDITOR
-		, ExtensionSection(InExtensionSection)
-#endif
 	{}
 
 	FName GetExtensionName() const
@@ -51,10 +41,7 @@ public:
 	}
 
 #if WITH_EDITOR
-	const FCEExtensionSection& GetExtensionSection() const
-	{
-		return ExtensionSection;
-	}
+	CLONEREFFECTOR_API FCEExtensionSection GetExtensionSection() const;
 #endif
 
 	/** Get the cloner component using this extension */
@@ -80,8 +67,13 @@ public:
 	/** Called when the meshes are updated */
 	virtual void OnClonerMeshesUpdated() {}
 
+	/** Updates all extensions parameters */
+	void UpdateExtensionParameters();
+
 	/** Request refresh extension next tick */
-	void UpdateExtensionParameters(bool bInUpdateCloner = true, bool bInImmediate = false);
+	void MarkExtensionDirty(bool bInUpdateCloner = true);
+
+	bool IsExtensionDirty() const;
 
 protected:
 	//~ Begin UObject
@@ -114,8 +106,5 @@ private:
 
 	bool bExtensionActive = false;
 
-#if WITH_EDITOR
-	/** Used for editor UI */
-	FCEExtensionSection ExtensionSection;
-#endif
+	ECEClonerSystemStatus ExtensionStatus = ECEClonerSystemStatus::UpToDate;
 };
\ No newline at end of file
diff --git a/Engine/Plugins/Experimental/ClonerEffector/Source/ClonerEffector/Public/Cloner/Extensions/CEClonerLifetimeExtension.h b/Engine/Plugins/Experimental/ClonerEffector/Source/ClonerEffector/Public/Cloner/Extensions/CEClonerLifetimeExtension.h
--- a/Engine/Plugins/Experimental/ClonerEffector/Source/ClonerEffector/Public/Cloner/Extensions/CEClonerLifetimeExtension.h
+++ b/Engine/Plugins/Experimental/ClonerEffector/Source/ClonerEffector/Public/Cloner/Extensions/CEClonerLifetimeExtension.h
@@ -10,7 +10,7 @@
 class UNiagaraDataInterfaceCurve;
 
 /** Extension dealing with clones lifetime options */
-UCLASS(MinimalAPI, BlueprintType, Within=CEClonerComponent)
+UCLASS(MinimalAPI, BlueprintType, Within=CEClonerComponent, meta=(Section="Emission", Priority=3))
 class UCEClonerLifetimeExtension : public UCEClonerExtensionBase
 {
 	GENERATED_BODY()
diff --git a/Engine/Plugins/Experimental/ClonerEffector/Source/ClonerEffector/Public/Cloner/Extensions/CEClonerMeshRendererExtension.h b/Engine/Plugins/Experimental/ClonerEffector/Source/ClonerEffector/Public/Cloner/Extensions/CEClonerMeshRendererExtension.h
--- a/Engine/Plugins/Experimental/ClonerEffector/Source/ClonerEffector/Public/Cloner/Extensions/CEClonerMeshRendererExtension.h
+++ b/Engine/Plugins/Experimental/ClonerEffector/Source/ClonerEffector/Public/Cloner/Extensions/CEClonerMeshRendererExtension.h
@@ -9,7 +9,7 @@
 #include "CEClonerMeshRendererExtension.generated.h"
 
 /** Extension dealing with mesh rendering options */
-UCLASS(MinimalAPI, BlueprintType, Within=CEClonerComponent)
+UCLASS(MinimalAPI, BlueprintType, Within=CEClonerComponent, meta=(Section="Rendering", Priority=4))
 class UCEClonerMeshRendererExtension : public UCEClonerExtensionBase
 {
 	GENERATED_BODY()
diff --git a/Engine/Plugins/Experimental/ClonerEffector/Source/ClonerEffector/Public/Cloner/Extensions/CEClonerProgressExtension.h b/Engine/Plugins/Experimental/ClonerEffector/Source/ClonerEffector/Public/Cloner/Extensions/CEClonerProgressExtension.h
--- a/Engine/Plugins/Experimental/ClonerEffector/Source/ClonerEffector/Public/Cloner/Extensions/CEClonerProgressExtension.h
+++ b/Engine/Plugins/Experimental/ClonerEffector/Source/ClonerEffector/Public/Cloner/Extensions/CEClonerProgressExtension.h
@@ -7,7 +7,7 @@
 #include "CEClonerProgressExtension.generated.h"
 
 /** Extension dealing with clone progress options */
-UCLASS(MinimalAPI, BlueprintType, Within=CEClonerComponent)
+UCLASS(MinimalAPI, BlueprintType, Within=CEClonerComponent, meta=(Section="Emission", Priority=4))
 class UCEClonerProgressExtension : public UCEClonerExtensionBase
 {
 	GENERATED_BODY()
@@ -17,9 +17,6 @@ public:
 		: UCEClonerExtensionBase(
 			TEXT("Progress")
 			, 0
-#if WITH_EDITOR
-			, UE::ClonerEffector::ClonerSection::EmissionSection
-#endif
 		)
 	{}
 
diff --git a/Engine/Plugins/Experimental/ClonerEffector/Source/ClonerEffector/Public/Cloner/Extensions/CEClonerRangeExtension.h b/Engine/Plugins/Experimental/ClonerEffector/Source/ClonerEffector/Public/Cloner/Extensions/CEClonerRangeExtension.h
--- a/Engine/Plugins/Experimental/ClonerEffector/Source/ClonerEffector/Public/Cloner/Extensions/CEClonerRangeExtension.h
+++ b/Engine/Plugins/Experimental/ClonerEffector/Source/ClonerEffector/Public/Cloner/Extensions/CEClonerRangeExtension.h
@@ -7,7 +7,7 @@
 #include "CEClonerRangeExtension.generated.h"
 
 /** Extension dealing with range options */
-UCLASS(MinimalAPI, BlueprintType, Within=CEClonerComponent)
+UCLASS(MinimalAPI, BlueprintType, Within=CEClonerComponent, meta=(Section="Cloner", Priority=0))
 class UCEClonerRangeExtension : public UCEClonerExtensionBase
 {
 	GENERATED_BODY()
@@ -17,9 +17,6 @@ public:
 		: UCEClonerExtensionBase(
 			TEXT("Range")
 			, 0
-#if WITH_EDITOR
-			, UE::ClonerEffector::ClonerSection::ClonerSection
-#endif
 		)
 	{}
 
diff --git a/Engine/Plugins/Experimental/ClonerEffector/Source/ClonerEffector/Public/Cloner/Extensions/CEClonerStepExtension.h b/Engine/Plugins/Experimental/ClonerEffector/Source/ClonerEffector/Public/Cloner/Extensions/CEClonerStepExtension.h
--- a/Engine/Plugins/Experimental/ClonerEffector/Source/ClonerEffector/Public/Cloner/Extensions/CEClonerStepExtension.h
+++ b/Engine/Plugins/Experimental/ClonerEffector/Source/ClonerEffector/Public/Cloner/Extensions/CEClonerStepExtension.h
@@ -7,7 +7,7 @@
 #include "CEClonerStepExtension.generated.h"
 
 /** Extension dealing with delta step accumulated options */
-UCLASS(MinimalAPI, BlueprintType, Within=CEClonerComponent)
+UCLASS(MinimalAPI, BlueprintType, Within=CEClonerComponent, meta=(Section="Cloner", Priority=0))
 class UCEClonerStepExtension : public UCEClonerExtensionBase
 {
 	GENERATED_BODY()
@@ -17,9 +17,6 @@ public:
 		: UCEClonerExtensionBase(
 			TEXT("Step")
 			, 0
-#if WITH_EDITOR
-			, UE::ClonerEffector::ClonerSection::ClonerSection
-#endif
 		)
 	{}
 
diff --git a/Engine/Plugins/Experimental/ClonerEffector/Source/ClonerEffector/Public/Cloner/Layouts/CEClonerGridLayout.h b/Engine/Plugins/Experimental/ClonerEffector/Source/ClonerEffector/Public/Cloner/Layouts/CEClonerGridLayout.h
--- a/Engine/Plugins/Experimental/ClonerEffector/Source/ClonerEffector/Public/Cloner/Layouts/CEClonerGridLayout.h
+++ b/Engine/Plugins/Experimental/ClonerEffector/Source/ClonerEffector/Public/Cloner/Layouts/CEClonerGridLayout.h
@@ -131,10 +131,10 @@ public:
 	}
 
 	UFUNCTION(BlueprintCallable, Category="Cloner|Layout|Grid")
-	CLONEREFFECTOR_API void SetTwistAxis(ENiagaraOrientationAxis InAxis);
+	CLONEREFFECTOR_API void SetTwistAxis(ECEClonerAxis InAxis);
 
 	UFUNCTION(BlueprintPure, Category="Cloner|Layout|Grid")
-	ENiagaraOrientationAxis GetTwistAxis() const
+	ECEClonerAxis GetTwistAxis() const
 	{
 		return TwistAxis;
 	}
@@ -186,8 +186,8 @@ protected:
 	UPROPERTY(EditInstanceOnly, BlueprintReadWrite, Setter, Getter, Category="Layout", meta=(UIMin="0", UIMax="100"))
 	float TwistFactor = 0.f;
 
-	UPROPERTY(EditInstanceOnly, BlueprintReadWrite, Setter, Getter, Category="Layout")
-	ENiagaraOrientationAxis TwistAxis = ENiagaraOrientationAxis::YAxis;
+	UPROPERTY(EditInstanceOnly, BlueprintReadWrite, Setter, Getter, Category="Layout", meta=(InvalidEnumValues="Custom"))
+	ECEClonerAxis TwistAxis = ECEClonerAxis::Y;
 
 private:
 #if WITH_EDITOR
diff --git a/Engine/Plugins/Experimental/ClonerEffector/Source/ClonerEffector/Public/Cloner/Layouts/CEClonerHoneycombLayout.h b/Engine/Plugins/Experimental/ClonerEffector/Source/ClonerEffector/Public/Cloner/Layouts/CEClonerHoneycombLayout.h
--- a/Engine/Plugins/Experimental/ClonerEffector/Source/ClonerEffector/Public/Cloner/Layouts/CEClonerHoneycombLayout.h
+++ b/Engine/Plugins/Experimental/ClonerEffector/Source/ClonerEffector/Public/Cloner/Layouts/CEClonerHoneycombLayout.h
@@ -95,10 +95,10 @@ public:
 	}
 
 	UFUNCTION(BlueprintCallable, Category="Cloner|Layout|Honeycomb")
-	CLONEREFFECTOR_API void SetTwistAxis(ENiagaraOrientationAxis InAxis);
+	CLONEREFFECTOR_API void SetTwistAxis(ECEClonerAxis InAxis);
 
 	UFUNCTION(BlueprintPure, Category="Cloner|Layout|Honeycomb")
-	ENiagaraOrientationAxis GetTwistAxis() const
+	ECEClonerAxis GetTwistAxis() const
 	{
 		return TwistAxis;
 	}
@@ -140,8 +140,8 @@ protected:
 	UPROPERTY(EditInstanceOnly, BlueprintReadWrite, Setter, Getter, Category="Layout", meta=(UIMin="0", UIMax="100"))
 	float TwistFactor = 0.f;
 
-	UPROPERTY(EditInstanceOnly, BlueprintReadWrite, Setter, Getter, Category="Layout")
-	ENiagaraOrientationAxis TwistAxis = ENiagaraOrientationAxis::YAxis;
+	UPROPERTY(EditInstanceOnly, BlueprintReadWrite, Setter, Getter, Category="Layout", meta=(InvalidEnumValues="Custom"))
+	ECEClonerAxis TwistAxis = ECEClonerAxis::Y;
 
 private:
 #if WITH_EDITOR
diff --git a/Engine/Plugins/Experimental/ClonerEffector/Source/ClonerEffector/Public/Cloner/Layouts/CEClonerLayoutBase.h b/Engine/Plugins/Experimental/ClonerEffector/Source/ClonerEffector/Public/Cloner/Layouts/CEClonerLayoutBase.h
--- a/Engine/Plugins/Experimental/ClonerEffector/Source/ClonerEffector/Public/Cloner/Layouts/CEClonerLayoutBase.h
+++ b/Engine/Plugins/Experimental/ClonerEffector/Source/ClonerEffector/Public/Cloner/Layouts/CEClonerLayoutBase.h
@@ -73,13 +73,10 @@ public:
 
 	AActor* GetClonerActor() const;
 
-	/** Request refresh layout next tick */
-	void UpdateLayoutParameters(bool bInUpdateCloner = true, bool bInImmediate = false);
-
-	/** Updates the cloner, forcing a reset of the system */
-	void RequestClonerUpdate(bool bInImmediate = false) const;
+	/** Updates all parameters handled by this layout */
+	void UpdateLayoutParameters();
 
-	/* Checks if the niagara system asset is valid and usable with the cloner */
+	/** Checks if the niagara system asset is valid and usable with the cloner */
 	bool IsLayoutValid() const;
 
 	/** Is this layout system cached and ready to be used */
@@ -108,7 +105,20 @@ public:
 	/** Gets the cloner extensions supported by this layout */
 	virtual TSet<FName> GetSupportedExtensions() const;
 
+	/** Request refresh layout next tick */
+	void MarkLayoutDirty(bool bInUpdateCloner = true);
+
+	/** Is the cloner not up to date with layout parameters */
+	bool IsLayoutDirty() const;
+
 protected:
+	//~ Begin UObject
+	virtual void PostEditImport() override;
+#if WITH_EDITOR
+	virtual void PostEditUndo() override;
+#endif
+	//~ End UObject
+
 	/** Called once after layout is loaded */
 	virtual void OnLayoutLoaded() {}
 
@@ -153,4 +163,6 @@ private:
 	FCEClonerEffectorDataInterfaces DataInterfaces;
 
 	int32 LoadRequestIdentifier = INDEX_NONE;
+
+	ECEClonerSystemStatus LayoutStatus = ECEClonerSystemStatus::UpToDate;
 };
\ No newline at end of file
diff --git a/Engine/Plugins/Experimental/ClonerEffector/Source/ClonerEffector/Public/Effector/CEEffectorComponent.h b/Engine/Plugins/Experimental/ClonerEffector/Source/ClonerEffector/Public/Effector/CEEffectorComponent.h
--- a/Engine/Plugins/Experimental/ClonerEffector/Source/ClonerEffector/Public/Effector/CEEffectorComponent.h
+++ b/Engine/Plugins/Experimental/ClonerEffector/Source/ClonerEffector/Public/Effector/CEEffectorComponent.h
@@ -30,6 +30,12 @@ class UCEEffectorComponent : public USceneComponent
 	friend class UCEClonerEffectorExtension;
 
 public:
+#if WITH_EDITOR
+	CLONEREFFECTOR_API static FName GetModeNamePropertyName();
+
+	CLONEREFFECTOR_API static FName GetTypeNamePropertyName();
+#endif
+
 	UCEEffectorComponent();
 
 	UFUNCTION(BlueprintCallable, Category="Effector")
@@ -140,6 +146,11 @@ public:
 		return ActiveMode;
 	}
 
+	TConstArrayView<TObjectPtr<UCEEffectorExtensionBase>> GetActiveExtensions() const
+	{
+		return ActiveExtensions;
+	}
+
 	void RequestClonerUpdate(bool bInImmediate);
 
 protected:
@@ -180,7 +191,7 @@ protected:
 	FName TypeName = NAME_None;
 
 	/** Cached active type used for faster access */
-	UPROPERTY(VisibleInstanceOnly, BlueprintReadOnly, Instanced, Transient, DuplicateTransient, NonTransactional, Category="Shape", meta=(DisplayAfter="TypeName"))
+	UPROPERTY(Transient, DuplicateTransient, NonTransactional)
 	TObjectPtr<UCEEffectorTypeBase> ActiveType;
 
 	/** Name of the shape type to use */
@@ -188,11 +199,11 @@ protected:
 	FName ModeName = NAME_None;
 
 	/** Cached active mode used for faster access */
-	UPROPERTY(VisibleInstanceOnly, BlueprintReadOnly, Instanced, Transient, DuplicateTransient, NonTransactional, Category="Mode", meta=(DisplayAfter="ModeName"))
+	UPROPERTY(Transient, DuplicateTransient, NonTransactional)
 	TObjectPtr<UCEEffectorModeBase> ActiveMode;
 
 	/** Active Extensions on this effector */
-	UPROPERTY(VisibleInstanceOnly, BlueprintReadOnly, Instanced, Transient, DuplicateTransient, NonTransactional, Category="Extensions")
+	UPROPERTY(Transient, DuplicateTransient, NonTransactional)
 	TArray<TObjectPtr<UCEEffectorExtensionBase>> ActiveExtensions;
 
 #if WITH_EDITORONLY_DATA
diff --git a/Engine/Plugins/Experimental/ClonerEffector/Source/ClonerEffector/Public/Effector/CEEffectorExtensionBase.h b/Engine/Plugins/Experimental/ClonerEffector/Source/ClonerEffector/Public/Effector/CEEffectorExtensionBase.h
--- a/Engine/Plugins/Experimental/ClonerEffector/Source/ClonerEffector/Public/Effector/CEEffectorExtensionBase.h
+++ b/Engine/Plugins/Experimental/ClonerEffector/Source/ClonerEffector/Public/Effector/CEEffectorExtensionBase.h
@@ -18,22 +18,11 @@ public:
 	UCEEffectorExtensionBase()
 		: UCEEffectorExtensionBase(
 			NAME_None
-#if WITH_EDITOR
-			, FCEExtensionSection(NAME_None, INDEX_NONE)
-#endif
 		)
 	{}
 
-	UCEEffectorExtensionBase(
-		FName InExtensionName
-#if WITH_EDITOR
-		, const FCEExtensionSection& InExtensionSection
-#endif
-		)
+	UCEEffectorExtensionBase(FName InExtensionName)
 		: ExtensionName(InExtensionName)
-#if WITH_EDITOR
-		, ExtensionSection(InExtensionSection)
-#endif
 	{}
 
 	FName GetExtensionName() const
@@ -42,17 +31,14 @@ public:
 	}
 
 #if WITH_EDITOR
-	const FCEExtensionSection& GetExtensionSection() const
-	{
-		return ExtensionSection;
-	}
+	CLONEREFFECTOR_API FCEExtensionSection GetExtensionSection() const;
 #endif
 
 	/** Get the effector component using this extension */
 	UCEEffectorComponent* GetEffectorComponent() const;
 
 	/** Request refresh extension next tick */
-	void UpdateExtensionParameters(bool bInUpdateLinkedCloners = false, bool bInImmediate = false);
+	void UpdateExtensionParameters(bool bInUpdateLinkedCloners = false);
 
 	/** Enable this extension */
 	void ActivateExtension();
@@ -88,8 +74,4 @@ private:
 	FName ExtensionName = NAME_None;
 
 	bool bExtensionActive = false;
-
-#if WITH_EDITOR
-	FCEExtensionSection ExtensionSection;
-#endif
 };
\ No newline at end of file
diff --git a/Engine/Plugins/Experimental/ClonerEffector/Source/ClonerEffector/Public/Effector/Extensions/CEEffectorForceExtension.h b/Engine/Plugins/Experimental/ClonerEffector/Source/ClonerEffector/Public/Effector/Extensions/CEEffectorForceExtension.h
--- a/Engine/Plugins/Experimental/ClonerEffector/Source/ClonerEffector/Public/Effector/Extensions/CEEffectorForceExtension.h
+++ b/Engine/Plugins/Experimental/ClonerEffector/Source/ClonerEffector/Public/Effector/Extensions/CEEffectorForceExtension.h
@@ -9,7 +9,7 @@
 class UCEEffectorComponent;
 
 /** Effector supported forces to affect clones */
-UCLASS(MinimalAPI, BlueprintType, Within=CEEffectorComponent)
+UCLASS(MinimalAPI, BlueprintType, Within=CEEffectorComponent, meta=(Section="Forces", Priority=3))
 class UCEEffectorForceExtension : public UCEEffectorExtensionBase
 {
 	GENERATED_BODY()
@@ -18,9 +18,6 @@ public:
 	UCEEffectorForceExtension()
 		: UCEEffectorExtensionBase(
 			TEXT("Forces")
-#if WITH_EDITOR
-			, UE::ClonerEffector::EffectorSection::ForcesSection
-#endif
 		)
 	{}
 
diff --git a/Engine/Plugins/Experimental/ClonerEffector/Source/ClonerEffector/Public/Effector/Modes/CEEffectorModeBase.h b/Engine/Plugins/Experimental/ClonerEffector/Source/ClonerEffector/Public/Effector/Modes/CEEffectorModeBase.h
--- a/Engine/Plugins/Experimental/ClonerEffector/Source/ClonerEffector/Public/Effector/Modes/CEEffectorModeBase.h
+++ b/Engine/Plugins/Experimental/ClonerEffector/Source/ClonerEffector/Public/Effector/Modes/CEEffectorModeBase.h
@@ -8,7 +8,7 @@
 class UCEEffectorComponent;
 
 /** Represents a behavior for an effector to affect clones in a specific way */
-UCLASS(MinimalAPI, Abstract, BlueprintType, Within=CEEffectorComponent)
+UCLASS(MinimalAPI, Abstract, BlueprintType, Within=CEEffectorComponent, meta=(Section="Mode", Priority=2))
 class UCEEffectorModeBase : public UCEEffectorExtensionBase
 {
 	GENERATED_BODY()
@@ -21,9 +21,6 @@ public:
 	UCEEffectorModeBase(FName InModeName, int32 InModeIdentifier)
 		: UCEEffectorExtensionBase(
 			InModeName
-#if WITH_EDITOR
-			, UE::ClonerEffector::EffectorSection::ModeSection
-#endif
 		)
 		, ModeIdentifier(InModeIdentifier)
 	{}
diff --git a/Engine/Plugins/Experimental/ClonerEffector/Source/ClonerEffector/Public/Effector/Types/CEEffectorTypeBase.h b/Engine/Plugins/Experimental/ClonerEffector/Source/ClonerEffector/Public/Effector/Types/CEEffectorTypeBase.h
--- a/Engine/Plugins/Experimental/ClonerEffector/Source/ClonerEffector/Public/Effector/Types/CEEffectorTypeBase.h
+++ b/Engine/Plugins/Experimental/ClonerEffector/Source/ClonerEffector/Public/Effector/Types/CEEffectorTypeBase.h
@@ -15,7 +15,7 @@ class UCEEffectorComponent;
 class UDynamicMesh;
 
 /** Represents a shape for an effector to affect clones on specific zones */
-UCLASS(MinimalAPI, Abstract, BlueprintType, Within=CEEffectorComponent)
+UCLASS(MinimalAPI, Abstract, BlueprintType, Within=CEEffectorComponent, meta=(Section="Shape", Priority=1))
 class UCEEffectorTypeBase : public UCEEffectorExtensionBase
 {
 	GENERATED_BODY()
@@ -31,9 +31,6 @@ public:
 	UCEEffectorTypeBase(FName InTypeName, int32 InTypeIdentifier)
 		: UCEEffectorExtensionBase(
 			InTypeName
-#if WITH_EDITOR
-			, UE::ClonerEffector::EffectorSection::ShapeSection
-#endif
 		)
 		, TypeIdentifier(InTypeIdentifier)
 	{}
diff --git a/Engine/Plugins/Experimental/ClonerEffector/Source/ClonerEffectorEditor/Private/CEEditorModule.cpp b/Engine/Plugins/Experimental/ClonerEffector/Source/ClonerEffectorEditor/Private/CEEditorModule.cpp
--- a/Engine/Plugins/Experimental/ClonerEffector/Source/ClonerEffectorEditor/Private/CEEditorModule.cpp
+++ b/Engine/Plugins/Experimental/ClonerEffector/Source/ClonerEffectorEditor/Private/CEEditorModule.cpp
@@ -7,9 +7,11 @@
 #include "Cloner/Customizations/CEEditorClonerActorDetailCustomization.h"
 #include "Cloner/Customizations/CEEditorClonerComponentDetailCustomization.h"
 #include "Cloner/Customizations/CEEditorClonerEffectorExtensionDetailCustomization.h"
+#include "Cloner/Customizations/CEEditorClonerLifetimeExtensionDetailCustomization.h"
 #include "Cloner/Customizations/CEEditorClonerMeshLayoutDetailCustomization.h"
 #include "Cloner/Customizations/CEEditorClonerSplineLayoutDetailCustomization.h"
 #include "Cloner/Extensions/CEClonerEffectorExtension.h"
+#include "Cloner/Extensions/CEClonerLifetimeExtension.h"
 #include "Cloner/Layouts/CEClonerMeshLayout.h"
 #include "Cloner/Layouts/CEClonerSplineLayout.h"
 #include "Effector/Customizations/CEEditorEffectorComponentDetailCustomization.h"
@@ -35,6 +37,7 @@ void FCEEditorModule::StartupModule()
 	PropertyModule.RegisterCustomClassLayout(UCEClonerEffectorExtension::StaticClass()->GetFName(), FOnGetDetailCustomizationInstance::CreateStatic(&FCEEditorClonerEffectorExtensionDetailCustomization::MakeInstance));
 	PropertyModule.RegisterCustomClassLayout(UCEClonerSplineLayout::StaticClass()->GetFName(), FOnGetDetailCustomizationInstance::CreateStatic(&FCEEditorClonerSplineLayoutDetailCustomization::MakeInstance));
 	PropertyModule.RegisterCustomClassLayout(UCEClonerMeshLayout::StaticClass()->GetFName(), FOnGetDetailCustomizationInstance::CreateStatic(&FCEEditorClonerMeshLayoutDetailCustomization::MakeInstance));
+	PropertyModule.RegisterCustomClassLayout(UCEClonerLifetimeExtension::StaticClass()->GetFName(), FOnGetDetailCustomizationInstance::CreateStatic(&FCEEditorClonerLifetimeExtensionDetailCustomization::MakeInstance));
 
 	// Effector customization
 	PropertyModule.RegisterCustomClassLayout(ACEEffectorActor::StaticClass()->GetFName(), FOnGetDetailCustomizationInstance::CreateStatic(&FCEEditorEffectorActorDetailCustomization::MakeInstance));
@@ -44,7 +47,7 @@ void FCEEditorModule::StartupModule()
 
 void FCEEditorModule::ShutdownModule()
 {
-	if (FModuleManager::Get().IsModuleLoaded(PropertyEditorName))
+	if (FModuleManager::Get().IsModuleLoaded(PropertyEditorName) && UObjectInitialized())
 	{
 		FPropertyEditorModule& PropertyModule = FModuleManager::GetModuleChecked<FPropertyEditorModule>(PropertyEditorName);
 
@@ -54,6 +57,7 @@ void FCEEditorModule::ShutdownModule()
 		PropertyModule.UnregisterCustomClassLayout(UCEClonerEffectorExtension::StaticClass()->GetFName());
 		PropertyModule.UnregisterCustomClassLayout(UCEClonerSplineLayout::StaticClass()->GetFName());
 		PropertyModule.UnregisterCustomClassLayout(UCEClonerMeshLayout::StaticClass()->GetFName());
+		PropertyModule.UnregisterCustomClassLayout(UCEClonerLifetimeExtension::StaticClass()->GetFName());
 
 		// Effector customization
 		PropertyModule.UnregisterCustomClassLayout(ACEEffectorActor::StaticClass()->GetFName());
diff --git a/Engine/Plugins/Experimental/ClonerEffector/Source/ClonerEffectorEditor/Private/Cloner/Customizations/CEEditorClonerComponentDetailCustomization.cpp b/Engine/Plugins/Experimental/ClonerEffector/Source/ClonerEffectorEditor/Private/Cloner/Customizations/CEEditorClonerComponentDetailCustomization.cpp
--- a/Engine/Plugins/Experimental/ClonerEffector/Source/ClonerEffectorEditor/Private/Cloner/Customizations/CEEditorClonerComponentDetailCustomization.cpp
+++ b/Engine/Plugins/Experimental/ClonerEffector/Source/ClonerEffectorEditor/Private/Cloner/Customizations/CEEditorClonerComponentDetailCustomization.cpp
@@ -2,18 +2,15 @@
 
 #include "Cloner/Customizations/CEEditorClonerComponentDetailCustomization.h"
 
-#include "CEEditorClonerEffectorExtensionDetailCustomization.h"
 #include "Cloner/CEClonerComponent.h"
-#include "Cloner/Extensions/CEClonerEffectorExtension.h"
 #include "Cloner/Extensions/CEClonerExtensionBase.h"
-#include "Cloner/Extensions/CEClonerLifetimeExtension.h"
 #include "Cloner/Layouts/CEClonerLayoutBase.h"
 #include "DetailLayoutBuilder.h"
 #include "DetailWidgetRow.h"
 #include "IDetailChildrenBuilder.h"
 #include "Input/Reply.h"
+#include "IPropertyUtilities.h"
 #include "Modules/ModuleManager.h"
-#include "NiagaraDataInterfaceCurve.h"
 #include "PropertyEditorModule.h"
 #include "UObject/Class.h"
 #include "UObject/Object.h"
@@ -32,131 +29,84 @@ void FCEEditorClonerComponentDetailCustomization::CustomizeDetails(IDetailLayout
 	// Remove niagara utilities
 	InDetailBuilder.HideCategory(TEXT("NiagaraComponent_Utilities"));
 
-	// Remove extension array property
-	TSharedRef<IPropertyHandle> ActiveExtensionsProperty = InDetailBuilder.GetProperty(UCEClonerComponent::GetActiveExtensionsName(), UCEClonerComponent::StaticClass());
+	TSharedRef<IPropertyUtilities> PropertyUtilities = InDetailBuilder.GetPropertyUtilities();
+	PropertyUtilitiesWeak = PropertyUtilities.ToWeakPtr();
+	ClonerComponentsWeak = InDetailBuilder.GetObjectsOfTypeBeingCustomized<UCEClonerComponent>();
 
-	if (!ActiveExtensionsProperty->IsValidHandle())
+	// Place LayoutName property above all properties in the category
 	{
-		return;
-	}
+		TSharedRef<IPropertyHandle> LayoutHandle = InDetailBuilder.GetProperty(UCEClonerComponent::GetLayoutNamePropertyName(), UCEClonerComponent::StaticClass());
+		LayoutHandle->SetOnPropertyValueChangedWithData(TDelegate<void(const FPropertyChangedEvent&)>::CreateStatic(&FCEEditorClonerComponentDetailCustomization::OnPropertyChanged, PropertyUtilitiesWeak));
 
-	InDetailBuilder.HideProperty(ActiveExtensionsProperty);
-
-	uint32 NumElements = 0;
-	TSharedPtr<IPropertyHandleArray> ActiveExtensionArrayProperty = ActiveExtensionsProperty->AsArray();
-	ActiveExtensionArrayProperty->GetNumElements(NumElements);
+		IDetailCategoryBuilder& LayoutCategoryBuilder = InDetailBuilder.EditCategory(TEXT("Layout"), FText::FromName(TEXT("Layout")));
+		LayoutCategoryBuilder.AddProperty(LayoutHandle);
+	}
 
 	// Everything needs to be below Cloner category
-	int32 StartOrder = InDetailBuilder.EditCategory(TEXT("Cloner")).GetSortOrder() + 1;
+	const IDetailCategoryBuilder& ClonerCategoryBuilder = InDetailBuilder.EditCategory(TEXT("Cloner"), FText::FromName(TEXT("Cloner")));
+	const int32 ExtensionSortOrderOffset = 1;
+	const int32 StartSortOrder = ClonerCategoryBuilder.GetSortOrder() + 1;
 
-	TSharedRef<IPropertyHandle> ActiveLayoutProperty = InDetailBuilder.GetProperty(UCEClonerComponent::GetActiveLayoutName(), UCEClonerComponent::StaticClass());
-
-	if (!ActiveLayoutProperty->IsValidHandle())
+	// Group same class objects together so their properties are grouped in the details panel when multiple cloners are selected
+	struct FDetailsCategoryData
 	{
-		return;
-	}
-
-	const FName LayoutCategoryName = ActiveLayoutProperty->GetDefaultCategoryName();
-	IDetailCategoryBuilder& LayoutCategory = InDetailBuilder.EditCategory(LayoutCategoryName);
-	LayoutCategory.SetSortOrder(StartOrder++);
+		FName SectionName;
+		int32 SortOrder;
+		TArray<UObject*> Objects;
+	};
 
-	const TSharedRef<FPropertySection> ClonerSection = PropertyModule.FindOrCreateSection(ComponentClassName, UE::ClonerEffector::ClonerSection::ClonerSection.SectionName, FText::FromName(UE::ClonerEffector::ClonerSection::ClonerSection.SectionName));
-	ClonerSection->AddCategory(LayoutCategoryName);
-
-	const TArray<TWeakObjectPtr<UCEClonerComponent>> ClonerComponentsWeak = InDetailBuilder.GetObjectsOfTypeBeingCustomized<UCEClonerComponent>();
-
-	for (uint32 ElementIndex = 0; ElementIndex < NumElements; ElementIndex++)
+	TMap<FName, FDetailsCategoryData> CategoryToData;
+	for (const TWeakObjectPtr<UCEClonerComponent>& ClonerComponentWeak : ClonerComponentsWeak)
 	{
-		TSharedPtr<IPropertyHandle> ActiveExtensionProperty = ActiveExtensionArrayProperty->GetElement(ElementIndex);
-
-		UObject* ObjectValue;
-		FPropertyAccess::Result ReadResult = ActiveExtensionProperty->GetValue(ObjectValue);
+		const UCEClonerComponent* ClonerComponent = ClonerComponentWeak.Get();
 
-		if (ReadResult != FPropertyAccess::Success)
+		if (!ClonerComponent)
 		{
 			continue;
 		}
 
-		UCEClonerExtensionBase* ActiveExtension = Cast<UCEClonerExtensionBase>(ObjectValue);
-
-		if (!IsValid(ActiveExtension))
+		if (UCEClonerLayoutBase* ActiveLayout = ClonerComponent->GetActiveLayout())
 		{
-			continue;
+			FDetailsCategoryData& CategoryData = CategoryToData.FindOrAdd(TEXT("Layout"));
+			CategoryData.SectionName = TEXT("Cloner");
+			CategoryData.SortOrder = StartSortOrder;
+			CategoryData.Objects.Add(ActiveLayout);
 		}
 
-		IDetailCategoryBuilder& ExtensionCategoryBuilder = InDetailBuilder.EditCategory(ActiveExtension->GetExtensionName());
-		ExtensionCategoryBuilder.SetSortOrder(StartOrder + ActiveExtension->GetExtensionSection().SectionOrder);
-
-		const FName ExtensionSectionName = ActiveExtension->GetExtensionSection().SectionName;
-		const TSharedRef<FPropertySection> ExtensionSection = PropertyModule.FindOrCreateSection(ComponentClassName, ExtensionSectionName, FText::FromName(ExtensionSectionName));
-		ExtensionSection->AddCategory(ActiveExtension->GetExtensionName());
-
-		TFunction<void(const TSharedPtr<IPropertyHandle>&)> AddObjectProperty = [&ExtensionCategoryBuilder, &AddObjectProperty](const TSharedPtr<IPropertyHandle>& InPropertyHandle)
+		for (UCEClonerExtensionBase* ActiveExtension : ClonerComponent->GetActiveExtensions())
 		{
-			uint32 NumChildren = 0;
-			InPropertyHandle->GetNumChildren(NumChildren);
+			FDetailsCategoryData& CategoryData = CategoryToData.FindOrAdd(ActiveExtension->GetExtensionName());
+			CategoryData.SectionName = ActiveExtension->GetExtensionSection().SectionName;
+			CategoryData.SortOrder = StartSortOrder + ExtensionSortOrderOffset + ActiveExtension->GetExtensionSection().SectionOrder;
+			CategoryData.Objects.Add(ActiveExtension);
+		}
+	}
 
-			for (uint32 ChildrenIndex = 0; ChildrenIndex < NumChildren; ChildrenIndex++)
-			{
-				TSharedPtr<IPropertyHandle> ChildHandle = InPropertyHandle->GetChildHandle(ChildrenIndex);
-
-				if (ChildHandle->IsCategoryHandle() && NumChildren == 1)
-				{
-					AddObjectProperty(ChildHandle);
-				}
-				else
-				{
-					const FProperty* Property = ChildHandle->GetProperty();
-
-					/*
-					 * Since AddExternalObject doesn't add hidden properties (EditConditionHides) of the object but only the visible ones
-					 * We add them manually, but this also means the object customization is not used, we need to set it manually from here
-					*/
-					if (Property && Property->GetName() == UCEClonerEffectorExtension::GetEffectorActorsWeakName())
-					{
-						FCEEditorClonerEffectorExtensionDetailCustomization::CustomizeEffectorsProperty(ChildHandle.ToSharedRef(), ExtensionCategoryBuilder);
-					}
-					else
-					{
-						ExtensionCategoryBuilder.AddProperty(ChildHandle);
-					}
-				}
-			}
-		};
+	FAddPropertyParams AddParams;
+	AddParams.CreateCategoryNodes(false);
+	AddParams.HideRootObjectNode(true);
 
-		AddObjectProperty(ActiveExtensionProperty->GetChildHandle(0));
+	for (const TPair<FName,FDetailsCategoryData>& CategoryToObjectsPair : CategoryToData)
+	{
+		const FName CategoryName = CategoryToObjectsPair.Key;
 
-		if (UCEClonerLifetimeExtension* LifetimeExtension = Cast<UCEClonerLifetimeExtension>(ActiveExtension))
+		if (CategoryName.IsNone() || CategoryToObjectsPair.Value.Objects.IsEmpty())
 		{
-			UNiagaraDataInterfaceCurve* LifetimeScaleCurve = LifetimeExtension->GetLifetimeScaleCurveDI();
-
-			if (!LifetimeScaleCurve || ClonerComponentsWeak.Num() > 1)
-			{
-				return;
-			}
+			continue;
+		}
 
-			const TArray<UObject*> ShowObjects {LifetimeScaleCurve};
-			ExtensionCategoryBuilder.SetShowAdvanced(true);
+		IDetailCategoryBuilder& CategoryBuilder = InDetailBuilder.EditCategory(CategoryName, FText::FromName(CategoryName));
+		CategoryBuilder.SetSortOrder(CategoryToObjectsPair.Value.SortOrder);
 
-			FAddPropertyParams Params;
-			Params.HideRootObjectNode(true);
-			Params.CreateCategoryNodes(false);
+		const TSharedRef<FPropertySection> PropertySection = PropertyModule.FindOrCreateSection(ComponentClassName, CategoryToObjectsPair.Value.SectionName, FText::FromName(CategoryToObjectsPair.Value.SectionName));
+		PropertySection->AddCategory(CategoryName);
 
-			IDetailPropertyRow* CurveRow = ExtensionCategoryBuilder.AddExternalObjects(ShowObjects, EPropertyLocation::Advanced, Params);
+		if (IDetailPropertyRow* ObjectRow = CategoryBuilder.AddExternalObjects(CategoryToObjectsPair.Value.Objects, EPropertyLocation::Default, AddParams))
+		{
+			TSharedPtr<IPropertyHandle> ObjectPropertyHandle = ObjectRow->GetPropertyHandle();
 
-			TWeakObjectPtr<UCEClonerLifetimeExtension> LifetimeExtensionWeak(LifetimeExtension);
-			CurveRow->Visibility(MakeAttributeLambda([LifetimeExtensionWeak]()
-			{
-				if (UCEClonerLifetimeExtension* LifetimeExtension = LifetimeExtensionWeak.Get())
-				{
-					return LifetimeExtension->GetLifetimeEnabled()
-						&& LifetimeExtension->GetLifetimeScaleEnabled()
-						? EVisibility::Visible
-						: EVisibility::Collapsed;
-				}
-
-				return EVisibility::Visible;
-			}));
+			// Fix for EditConditionHides not appearing when condition is met due to AddExternalObjects not rebuilding children
+			ObjectPropertyHandle->SetOnChildPropertyValueChangedWithData(TDelegate<void(const FPropertyChangedEvent&)>::CreateStatic(&FCEEditorClonerComponentDetailCustomization::OnChildPropertyChanged, ObjectPropertyHandle->AsWeak()));
 		}
 	}
 
@@ -188,40 +138,6 @@ void FCEEditorClonerComponentDetailCustomization::CustomizeDetails(IDetailLayout
 				FunctionToCategory.Add(FunctionName, FName(Function->GetMetaData(TEXT("Category"))));
 			}
 		}
-
-		// Look for CallInEditor functions in active layout
-		if (UCEClonerLayoutBase* ActiveLayout = ClonerComponent->GetActiveLayout())
-		{
-			// Iterate through all UFunctions in the class
-			for (UFunction* Function : TFieldRange<UFunction>(ActiveLayout->GetClass(), EFieldIteratorFlags::ExcludeSuper))
-			{
-				// Only CallInEditor function with 0 parameters
-				if (Function && Function->HasMetaData("CallInEditor") && Function->NumParms == 0)
-				{
-					FName FunctionName = Function->GetFName();
-					TMap<TWeakObjectPtr<UObject>, TWeakObjectPtr<UFunction>>& ObjectFunctions = LayoutFunctionNames.FindOrAdd(FunctionName);
-					ObjectFunctions.Add(ActiveLayout, Function);
-					FunctionToCategory.Add(FunctionName, FName(Function->GetMetaData(TEXT("Category"))));
-				}
-			}
-		}
-
-		// Look for CallInEditor functions in active extensions
-		for (const TObjectPtr<UCEClonerExtensionBase>& ActiveExtension : ClonerComponent->GetActiveExtensions())
-		{
-			// Iterate through all UFunctions in the class
-			for (UFunction* Function : TFieldRange<UFunction>(ActiveExtension->GetClass(), EFieldIteratorFlags::ExcludeSuper))
-			{
-				// Only CallInEditor function with 0 parameters
-				if (Function && Function->HasMetaData("CallInEditor") && Function->NumParms == 0)
-				{
-					FName FunctionName = Function->GetFName();
-					TMap<TWeakObjectPtr<UObject>, TWeakObjectPtr<UFunction>>& ObjectFunctions = LayoutFunctionNames.FindOrAdd(FunctionName);
-					ObjectFunctions.Add(ActiveExtension, Function);
-					FunctionToCategory.Add(FunctionName, FName(Function->GetMetaData(TEXT("Category"))));
-				}
-			}
-		}
 	}
 
 	if (!LayoutFunctionNames.IsEmpty())
@@ -285,6 +201,45 @@ void FCEEditorClonerComponentDetailCustomization::RemoveEmptySections()
 	PropertyModule.RemoveSection(ComponentClassName, TEXT("Streaming"));
 }
 
+void FCEEditorClonerComponentDetailCustomization::OnPropertyChanged(const FPropertyChangedEvent& InEvent, TWeakPtr<IPropertyUtilities> InUtilitiesWeak)
+{
+	if (const TSharedPtr<IPropertyUtilities> PropertyUtilities = InUtilitiesWeak.Pin())
+	{
+		if (InEvent.ChangeType != EPropertyChangeType::Interactive)
+		{
+			PropertyUtilities->RequestForceRefresh();
+		}
+	}
+}
+
+void FCEEditorClonerComponentDetailCustomization::Init()
+{
+	// When using a layout for the first time, it is not yet loaded and the property change will trigger an async load
+	UCEClonerComponent::OnClonerLayoutLoaded().RemoveAll(this);
+	UCEClonerComponent::OnClonerLayoutLoaded().AddSP(this, &FCEEditorClonerComponentDetailCustomization::OnClonerLayoutLoaded);
+}
+
+void FCEEditorClonerComponentDetailCustomization::OnChildPropertyChanged(const FPropertyChangedEvent& InEvent, TWeakPtr<IPropertyHandle> InParentHandleWeak)
+{
+	if (const TSharedPtr<IPropertyHandle> PropertyHandle = InParentHandleWeak.Pin())
+	{
+		if (InEvent.ChangeType != EPropertyChangeType::Interactive)
+		{
+			PropertyHandle->RequestRebuildChildren();
+		}
+	}
+}
+
+void FCEEditorClonerComponentDetailCustomization::OnClonerLayoutLoaded(UCEClonerComponent* InCloner, UCEClonerLayoutBase* InLayout)
+{
+	const TSharedPtr<IPropertyUtilities> PropertyUtilities = PropertyUtilitiesWeak.Pin();
+
+	if (PropertyUtilities && InCloner && ClonerComponentsWeak.Contains(InCloner))
+	{
+		PropertyUtilities->RequestForceRefresh();
+	}
+}
+
 FReply FCEEditorClonerComponentDetailCustomization::OnFunctionButtonClicked(FName InFunctionName)
 {
 	if (TMap<TWeakObjectPtr<UObject>, TWeakObjectPtr<UFunction>> const* ObjectFunctions = LayoutFunctionNames.Find(InFunctionName))
diff --git a/Engine/Plugins/Experimental/ClonerEffector/Source/ClonerEffectorEditor/Private/Cloner/Customizations/CEEditorClonerComponentDetailCustomization.h b/Engine/Plugins/Experimental/ClonerEffector/Source/ClonerEffectorEditor/Private/Cloner/Customizations/CEEditorClonerComponentDetailCustomization.h
--- a/Engine/Plugins/Experimental/ClonerEffector/Source/ClonerEffectorEditor/Private/Cloner/Customizations/CEEditorClonerComponentDetailCustomization.h
+++ b/Engine/Plugins/Experimental/ClonerEffector/Source/ClonerEffectorEditor/Private/Cloner/Customizations/CEEditorClonerComponentDetailCustomization.h
@@ -6,8 +6,13 @@
 #include "IDetailCustomization.h"
 #include "UObject/NameTypes.h"
 #include "UObject/WeakObjectPtr.h"
+#include "Templates/SharedPointerFwd.h"
 
 class FReply;
+class IPropertyHandle;
+class IPropertyUtilities;
+class UCEClonerComponent;
+class UCEClonerLayoutBase;
 class UFunction;
 class UObject;
 
@@ -17,7 +22,9 @@ class FCEEditorClonerComponentDetailCustomization : public IDetailCustomization
 public:
 	static TSharedRef<IDetailCustomization> MakeInstance()
 	{
-		return MakeShared<FCEEditorClonerComponentDetailCustomization>();
+		TSharedRef<FCEEditorClonerComponentDetailCustomization> Customization = MakeShared<FCEEditorClonerComponentDetailCustomization>();
+		Customization->Init();
+		return Customization;
 	}
 
 	explicit FCEEditorClonerComponentDetailCustomization()
@@ -31,9 +38,24 @@ public:
 
 protected:
 	static void RemoveEmptySections();
+	static void OnChildPropertyChanged(const FPropertyChangedEvent& InEvent, TWeakPtr<IPropertyHandle> InParentHandleWeak);
+	static void OnPropertyChanged(const FPropertyChangedEvent& InEvent, TWeakPtr<IPropertyUtilities> InUtilitiesWeak);
+
+	/** Bind delegates needed */
+	void Init();
+
+	/** Used to refresh details view when layout changes */
+	void OnClonerLayoutLoaded(UCEClonerComponent* InCloner, UCEClonerLayoutBase* InLayout);
 
 	/** Execute ufunction with that name on selected objects */
 	FReply OnFunctionButtonClicked(FName InFunctionName);
 
+	/** Function name to object ufunction mapping */
 	TMap<FName, TMap<TWeakObjectPtr<UObject>, TWeakObjectPtr<UFunction>>> LayoutFunctionNames;
+
+	/** Customized objects */
+	TArray<TWeakObjectPtr<UCEClonerComponent>> ClonerComponentsWeak;
+
+	/** Property utilities for details view refresh */
+	TWeakPtr<IPropertyUtilities> PropertyUtilitiesWeak;
 };
diff --git a/Engine/Plugins/Experimental/ClonerEffector/Source/ClonerEffectorEditor/Private/Cloner/Customizations/CEEditorClonerEffectorExtensionDetailCustomization.h b/Engine/Plugins/Experimental/ClonerEffector/Source/ClonerEffectorEditor/Private/Cloner/Customizations/CEEditorClonerEffectorExtensionDetailCustomization.h
--- a/Engine/Plugins/Experimental/ClonerEffector/Source/ClonerEffectorEditor/Private/Cloner/Customizations/CEEditorClonerEffectorExtensionDetailCustomization.h
+++ b/Engine/Plugins/Experimental/ClonerEffector/Source/ClonerEffectorEditor/Private/Cloner/Customizations/CEEditorClonerEffectorExtensionDetailCustomization.h
@@ -18,10 +18,11 @@ public:
 		return MakeShared<FCEEditorClonerEffectorExtensionDetailCustomization>();
 	}
 
-	static void CustomizeEffectorsProperty(TSharedRef<IPropertyHandle> InProperty, IDetailCategoryBuilder& InCategoryBuilder);
-	static bool OnFilterEffectorActor(const AActor* InActor);
-
 	//~ Begin IDetailCustomization
 	virtual void CustomizeDetails(IDetailLayoutBuilder& InDetailBuilder) override;
 	//~ End IDetailCustomization
+
+protected:
+	static void CustomizeEffectorsProperty(TSharedRef<IPropertyHandle> InProperty, IDetailCategoryBuilder& InCategoryBuilder);
+	static bool OnFilterEffectorActor(const AActor* InActor);
 };
diff --git a/Engine/Plugins/Experimental/ClonerEffector/Source/ClonerEffectorEditor/Private/Cloner/Customizations/CEEditorClonerLifetimeExtensionDetailCustomization.cpp b/Engine/Plugins/Experimental/ClonerEffector/Source/ClonerEffectorEditor/Private/Cloner/Customizations/CEEditorClonerLifetimeExtensionDetailCustomization.cpp
new file mode 100644
--- /dev/null
+++ b/Engine/Plugins/Experimental/ClonerEffector/Source/ClonerEffectorEditor/Private/Cloner/Customizations/CEEditorClonerLifetimeExtensionDetailCustomization.cpp
@@ -0,0 +1,59 @@
+// Copyright Epic Games, Inc. All Rights Reserved.
+
+#include "Cloner/Customizations/CEEditorClonerLifetimeExtensionDetailCustomization.h"
+
+#include "Cloner/Extensions/CEClonerLifetimeExtension.h"
+#include "DetailBuilderTypes.h"
+#include "DetailCategoryBuilder.h"
+#include "DetailLayoutBuilder.h"
+#include "GameFramework/Actor.h"
+#include "NiagaraDataInterfaceCurve.h"
+#include "PropertyHandle.h"
+
+void FCEEditorClonerLifetimeExtensionDetailCustomization::CustomizeDetails(IDetailLayoutBuilder& InDetailBuilder)
+{
+	const TArray<TWeakObjectPtr<UCEClonerLifetimeExtension>> LifetimeExtensionsWeak = InDetailBuilder.GetObjectsOfTypeBeingCustomized<UCEClonerLifetimeExtension>();
+
+	FAddPropertyParams Params;
+	Params.HideRootObjectNode(true);
+	Params.CreateCategoryNodes(false);
+
+	for (const TWeakObjectPtr<UCEClonerLifetimeExtension>& LifetimeExtensionWeak : LifetimeExtensionsWeak)
+	{
+		const UCEClonerLifetimeExtension* LifetimeExtension = LifetimeExtensionWeak.Get();
+
+		if (!LifetimeExtension)
+		{
+			continue;
+		}
+
+		const FName CategoryName = LifetimeExtension->GetExtensionName();
+
+		IDetailCategoryBuilder& CategoryBuilder = InDetailBuilder.EditCategory(CategoryName);
+		CategoryBuilder.SetShowAdvanced(true);
+
+		/**
+		 * UNiagaraDataInterfaceCurve cannot display simultaneously multiple curves, so we need to add them separately
+		 */
+		if (IDetailPropertyRow* Row = CategoryBuilder.AddExternalObjects({LifetimeExtension->GetLifetimeScaleCurveDI()}, EPropertyLocation::Advanced, Params))
+		{
+			const TAttribute<EVisibility> VisibilityAttr = TAttribute<EVisibility>::CreateSP(this, &FCEEditorClonerLifetimeExtensionDetailCustomization::GetCurveVisibility, LifetimeExtensionWeak);
+			Row->Visibility(VisibilityAttr);
+		}
+	}
+}
+
+EVisibility FCEEditorClonerLifetimeExtensionDetailCustomization::GetCurveVisibility(TWeakObjectPtr<UCEClonerLifetimeExtension> InExtensionWeak) const
+{
+	const UCEClonerLifetimeExtension* Extension = InExtensionWeak.Get();
+
+	if (!Extension)
+	{
+		return EVisibility::Collapsed;
+	}
+
+	return Extension->GetLifetimeEnabled()
+		&& Extension->GetLifetimeScaleEnabled()
+		? EVisibility::Visible
+		: EVisibility::Collapsed;
+}
diff --git a/Engine/Plugins/Experimental/ClonerEffector/Source/ClonerEffectorEditor/Private/Cloner/Customizations/CEEditorClonerLifetimeExtensionDetailCustomization.h b/Engine/Plugins/Experimental/ClonerEffector/Source/ClonerEffectorEditor/Private/Cloner/Customizations/CEEditorClonerLifetimeExtensionDetailCustomization.h
new file mode 100644
--- /dev/null
+++ b/Engine/Plugins/Experimental/ClonerEffector/Source/ClonerEffectorEditor/Private/Cloner/Customizations/CEEditorClonerLifetimeExtensionDetailCustomization.h
@@ -0,0 +1,31 @@
+// Copyright Epic Games, Inc. All Rights Reserved.
+
+#pragma once
+
+#include "IDetailCustomization.h"
+#include "Templates/SharedPointer.h"
+
+class IDetailCategoryBuilder;
+class IPropertyHandle;
+class UCEClonerLifetimeExtension;
+struct EVisibility;
+
+/** Used to customize cloner lifetime extension properties in details panel */
+class FCEEditorClonerLifetimeExtensionDetailCustomization : public IDetailCustomization
+{
+public:
+	static TSharedRef<IDetailCustomization> MakeInstance()
+	{
+		return MakeShared<FCEEditorClonerLifetimeExtensionDetailCustomization>();
+	}
+
+	//~ Begin IDetailCustomization
+	virtual void CustomizeDetails(IDetailLayoutBuilder& InDetailBuilder) override;
+	//~ End IDetailCustomization
+
+protected:
+	EVisibility GetCurveVisibility(TWeakObjectPtr<UCEClonerLifetimeExtension> InExtensionWeak) const;
+
+	TSharedPtr<IPropertyHandle> LifetimeEnabledPropertyHandle;
+	TSharedPtr<IPropertyHandle> LifetimeScaleEnabledPropertyHandle;
+};
diff --git a/Engine/Plugins/Experimental/ClonerEffector/Source/ClonerEffectorEditor/Private/Effector/Customizations/CEEditorEffectorComponentDetailCustomization.cpp b/Engine/Plugins/Experimental/ClonerEffector/Source/ClonerEffectorEditor/Private/Effector/Customizations/CEEditorEffectorComponentDetailCustomization.cpp
--- a/Engine/Plugins/Experimental/ClonerEffector/Source/ClonerEffectorEditor/Private/Effector/Customizations/CEEditorEffectorComponentDetailCustomization.cpp
+++ b/Engine/Plugins/Experimental/ClonerEffector/Source/ClonerEffectorEditor/Private/Effector/Customizations/CEEditorEffectorComponentDetailCustomization.cpp
@@ -6,9 +6,11 @@
 #include "DetailWidgetRow.h"
 #include "Effector/CEEffectorComponent.h"
 #include "Effector/CEEffectorExtensionBase.h"
+#include "Effector/Modes/CEEffectorModeBase.h"
+#include "Effector/Types/CEEffectorTypeBase.h"
+#include "IPropertyUtilities.h"
 #include "Modules/ModuleManager.h"
 #include "PropertyEditorModule.h"
-#include "Widgets/Input/SButton.h"
 
 #define LOCTEXT_NAMESPACE "CEEditorEffectorComponentDetailCustomization"
 
@@ -17,114 +19,106 @@ void FCEEditorEffectorComponentDetailCustomization::CustomizeDetails(IDetailLayo
 	FPropertyEditorModule& PropertyModule = FModuleManager::GetModuleChecked<FPropertyEditorModule>(TEXT("PropertyEditor"));
 	const FName ComponentClassName = UCEEffectorComponent::StaticClass()->GetFName();
 
-	// Remove extension array property
-	TSharedRef<IPropertyHandle> ActiveExtensionsProperty = InDetailBuilder.GetProperty(GET_MEMBER_NAME_CHECKED(UCEEffectorComponent, ActiveExtensions), UCEEffectorComponent::StaticClass());
+	TSharedRef<IPropertyUtilities> PropertyUtilities = InDetailBuilder.GetPropertyUtilities();
 
-	if (!ActiveExtensionsProperty->IsValidHandle())
+	// Place TypeName property above all properties in the category
 	{
-		return;
-	}
-
-	InDetailBuilder.HideProperty(ActiveExtensionsProperty);
-
-	// Everything needs to be below Effector category
-	const FName EffectorCategoryName = TEXT("Effector");
-	int32 StartOrder = InDetailBuilder.EditCategory(EffectorCategoryName).GetSortOrder() + 1;
+		TSharedRef<IPropertyHandle> TypeHandle = InDetailBuilder.GetProperty(UCEEffectorComponent::GetTypeNamePropertyName(), UCEEffectorComponent::StaticClass());
+		TypeHandle->SetOnPropertyValueChangedWithData(TDelegate<void(const FPropertyChangedEvent&)>::CreateStatic(&FCEEditorEffectorComponentDetailCustomization::OnPropertyChanged, PropertyUtilities.ToWeakPtr()));
 
-	const FName EffectorSectionName = UE::ClonerEffector::EffectorSection::EffectorSection.SectionName;
-	const TSharedRef<FPropertySection> EffectorSection = PropertyModule.FindOrCreateSection(ComponentClassName, EffectorSectionName, FText::FromName(EffectorSectionName));
-	EffectorSection->AddCategory(EffectorCategoryName);
+		IDetailCategoryBuilder& ShapeCategoryBuilder = InDetailBuilder.EditCategory(TEXT("Shape"), FText::FromName(TEXT("Shape")));
+		ShapeCategoryBuilder.AddProperty(TypeHandle);
+	}
 
-	// Shape
+	// Place ModeName property above all properties in the category
 	{
-		TSharedRef<IPropertyHandle> ActiveTypeProperty = InDetailBuilder.GetProperty(GET_MEMBER_NAME_CHECKED(UCEEffectorComponent, ActiveType), UCEEffectorComponent::StaticClass());
+		TSharedRef<IPropertyHandle> ModeHandle = InDetailBuilder.GetProperty(UCEEffectorComponent::GetModeNamePropertyName(), UCEEffectorComponent::StaticClass());
+		ModeHandle->SetOnPropertyValueChangedWithData(TDelegate<void(const FPropertyChangedEvent&)>::CreateStatic(&FCEEditorEffectorComponentDetailCustomization::OnPropertyChanged, PropertyUtilities.ToWeakPtr()));
 
-		if (!ActiveTypeProperty->IsValidHandle())
-		{
-			return;
-		}
+		IDetailCategoryBuilder& ModeCategoryBuilder = InDetailBuilder.EditCategory(TEXT("Mode"), FText::FromName(TEXT("Mode")));
+		ModeCategoryBuilder.AddProperty(ModeHandle);
+	}
 
-		const FName TypeCategoryName = ActiveTypeProperty->GetDefaultCategoryName();
+	const TArray<TWeakObjectPtr<UCEEffectorComponent>> EffectorComponentsWeak = InDetailBuilder.GetObjectsOfTypeBeingCustomized<UCEEffectorComponent>();
 
-		IDetailCategoryBuilder& TypeCategory = InDetailBuilder.EditCategory(TypeCategoryName);
-		TypeCategory.SetSortOrder(StartOrder++);
+	// Everything needs to be below Effector category
+	const IDetailCategoryBuilder& EffectorCategoryBuilder = InDetailBuilder.EditCategory(TEXT("Effector"), FText::FromName(TEXT("Effector")));
+	const int32 ExtensionSortOrderOffset = 2;
+	const int32 StartSortOrder = EffectorCategoryBuilder.GetSortOrder() + 1;
 
-		const FName TypeSectionName = UE::ClonerEffector::EffectorSection::ShapeSection.SectionName;
-		const TSharedRef<FPropertySection> TypeSection = PropertyModule.FindOrCreateSection(ComponentClassName, TypeSectionName, FText::FromName(TypeSectionName));
-		TypeSection->AddCategory(TypeCategoryName);
-	}
+	// Group same class objects together so their properties are grouped in the details panel when multiple effectors are selected
+	struct FDetailsCategoryData
+	{
+		FName SectionName;
+		int32 SortOrder;
+		TArray<UObject*> Objects;
+	};
 
-	// Mode
+	TMap<FName, FDetailsCategoryData> CategoryToData;
+	for (const TWeakObjectPtr<UCEEffectorComponent>& EffectorComponentWeak : EffectorComponentsWeak)
 	{
-		TSharedRef<IPropertyHandle> ActiveModeProperty = InDetailBuilder.GetProperty(GET_MEMBER_NAME_CHECKED(UCEEffectorComponent, ActiveMode), UCEEffectorComponent::StaticClass());
+		const UCEEffectorComponent* EffectorComponent = EffectorComponentWeak.Get();
 
-		if (!ActiveModeProperty->IsValidHandle())
+		if (!EffectorComponent)
 		{
-			return;
+			continue;
 		}
 
-		const FName ModeCategoryName = ActiveModeProperty->GetDefaultCategoryName();
-
-		IDetailCategoryBuilder& ModeCategory = InDetailBuilder.EditCategory(ModeCategoryName);
-		ModeCategory.SetSortOrder(StartOrder++);
-
-		const FName ModeSectionName = UE::ClonerEffector::EffectorSection::ModeSection.SectionName;
-		const TSharedRef<FPropertySection> ModeSection = PropertyModule.FindOrCreateSection(ComponentClassName, ModeSectionName, FText::FromName(ModeSectionName));
-		ModeSection->AddCategory(ModeCategoryName);
-	}
-
-	TFunction<void(const TSharedPtr<IPropertyHandle>&, IDetailCategoryBuilder&)> AddObjectProperty = [&AddObjectProperty](const TSharedPtr<IPropertyHandle>& InPropertyHandle, IDetailCategoryBuilder& InCategoryBuilder)
-	{
-		uint32 NumChildren = 0;
-		InPropertyHandle->GetNumChildren(NumChildren);
-
-		for (uint32 ChildrenIndex = 0; ChildrenIndex < NumChildren; ChildrenIndex++)
+		if (UCEEffectorTypeBase* ActiveType = EffectorComponent->GetActiveType())
 		{
-			TSharedPtr<IPropertyHandle> ChildHandle = InPropertyHandle->GetChildHandle(ChildrenIndex);
-
-			if (ChildHandle->IsCategoryHandle() && NumChildren == 1)
-			{
-				AddObjectProperty(ChildHandle, InCategoryBuilder);
-			}
-			else
-			{
-				InCategoryBuilder.AddProperty(ChildHandle);
-			}
+			FDetailsCategoryData& CategoryData = CategoryToData.FindOrAdd(TEXT("Shape"));
+			CategoryData.SectionName = ActiveType->GetExtensionSection().SectionName;
+			CategoryData.SortOrder = StartSortOrder + ActiveType->GetExtensionSection().SectionOrder;
+			CategoryData.Objects.Add(ActiveType);
 		}
-	};
-
-	uint32 NumElements = 0;
-	TSharedPtr<IPropertyHandleArray> ActiveExtensionArrayProperty = ActiveExtensionsProperty->AsArray();
-	ActiveExtensionArrayProperty->GetNumElements(NumElements);
-
-	for (uint32 ElementIndex = 0; ElementIndex < NumElements; ElementIndex++)
-	{
-		TSharedPtr<IPropertyHandle> ActiveExtensionProperty = ActiveExtensionArrayProperty->GetElement(ElementIndex);
 
-		UObject* ObjectValue;
-		FPropertyAccess::Result ReadResult = ActiveExtensionProperty->GetValue(ObjectValue);
+		if (UCEEffectorModeBase* ActiveMode = EffectorComponent->GetActiveMode())
+		{
+			FDetailsCategoryData& CategoryData = CategoryToData.FindOrAdd(TEXT("Mode"));
+			CategoryData.SectionName = ActiveMode->GetExtensionSection().SectionName;
+			CategoryData.SortOrder = StartSortOrder + ActiveMode->GetExtensionSection().SectionOrder;
+			CategoryData.Objects.Add(ActiveMode);
+		}
 
-		if (ReadResult != FPropertyAccess::Success)
+		for (const TObjectPtr<UCEEffectorExtensionBase>& ActiveExtension : EffectorComponent->GetActiveExtensions())
 		{
-			continue;
+			FDetailsCategoryData& CategoryData = CategoryToData.FindOrAdd(ActiveExtension->GetExtensionName());
+			CategoryData.SectionName = ActiveExtension->GetExtensionSection().SectionName;
+			CategoryData.SortOrder = StartSortOrder + ExtensionSortOrderOffset + ActiveExtension->GetExtensionSection().SectionOrder;
+			CategoryData.Objects.Add(ActiveExtension);
 		}
+	}
+
+	const FName EffectorSectionName = TEXT("Effector");
+	const TSharedRef<FPropertySection> EffectorSection = PropertyModule.FindOrCreateSection(ComponentClassName, EffectorSectionName, FText::FromName(EffectorSectionName));
+	EffectorSection->AddCategory(TEXT("Effector"));
 
-		UCEEffectorExtensionBase* ActiveExtension = Cast<UCEEffectorExtensionBase>(ObjectValue);
+	FAddPropertyParams AddParams;
+	AddParams.CreateCategoryNodes(false);
+	AddParams.HideRootObjectNode(true);
 
-		if (!ActiveExtension)
+	for (const TPair<FName,FDetailsCategoryData>& CategoryToObjectsPair : CategoryToData)
+	{
+		const FName CategoryName = CategoryToObjectsPair.Key;
+
+		if (CategoryName.IsNone() || CategoryToObjectsPair.Value.Objects.IsEmpty())
 		{
 			continue;
 		}
 
-		const FName ExtensionSectionName = ActiveExtension->GetExtensionSection().SectionName;
+		IDetailCategoryBuilder& CategoryBuilder = InDetailBuilder.EditCategory(CategoryName, FText::FromName(CategoryName));
+		CategoryBuilder.SetSortOrder(CategoryToObjectsPair.Value.SortOrder);
 
-		IDetailCategoryBuilder& ExtensionCategory = InDetailBuilder.EditCategory(ExtensionSectionName);
-		ExtensionCategory.SetSortOrder(StartOrder + ActiveExtension->GetExtensionSection().SectionOrder);
+		const TSharedRef<FPropertySection> PropertySection = PropertyModule.FindOrCreateSection(ComponentClassName, CategoryToObjectsPair.Value.SectionName, FText::FromName(CategoryToObjectsPair.Value.SectionName));
+		PropertySection->AddCategory(CategoryName);
 
-		const TSharedRef<FPropertySection> ExtensionSection = PropertyModule.FindOrCreateSection(ComponentClassName, ExtensionSectionName, FText::FromName(ExtensionSectionName));
-		ExtensionSection->AddCategory(ExtensionSectionName);
+		if (IDetailPropertyRow* ObjectRow = CategoryBuilder.AddExternalObjects(CategoryToObjectsPair.Value.Objects, EPropertyLocation::Default, AddParams))
+		{
+			TSharedPtr<IPropertyHandle> ObjectPropertyHandle = ObjectRow->GetPropertyHandle();
 
-		AddObjectProperty(ActiveExtensionProperty->GetChildHandle(0), ExtensionCategory);
+			// Fix for EditConditionHides not appearing when condition is met due to AddExternalObjects not rebuilding children
+			ObjectPropertyHandle->SetOnChildPropertyValueChangedWithData(TDelegate<void(const FPropertyChangedEvent&)>::CreateStatic(&FCEEditorEffectorComponentDetailCustomization::OnChildPropertyChanged, ObjectPropertyHandle->AsWeak()));
+		}
 	}
 }
 
@@ -137,4 +131,26 @@ void FCEEditorEffectorComponentDetailCustomization::RemoveEmptySections()
 	PropertyModule.RemoveSection(ComponentClassName, TEXT("Streaming"));
 }
 
+void FCEEditorEffectorComponentDetailCustomization::OnPropertyChanged(const FPropertyChangedEvent& InEvent, TWeakPtr<IPropertyUtilities> InUtilitiesWeak)
+{
+	if (const TSharedPtr<IPropertyUtilities> PropertyUtilities = InUtilitiesWeak.Pin())
+	{
+		if (InEvent.ChangeType != EPropertyChangeType::Interactive)
+		{
+			PropertyUtilities->RequestForceRefresh();
+		}
+	}
+}
+
+void FCEEditorEffectorComponentDetailCustomization::OnChildPropertyChanged(const FPropertyChangedEvent& InEvent, TWeakPtr<IPropertyHandle> InParentHandleWeak)
+{
+	if (const TSharedPtr<IPropertyHandle> PropertyHandle = InParentHandleWeak.Pin())
+	{
+		if (InEvent.ChangeType != EPropertyChangeType::Interactive)
+		{
+			PropertyHandle->RequestRebuildChildren();
+		}
+	}
+}
+
 #undef LOCTEXT_NAMESPACE
diff --git a/Engine/Plugins/Experimental/ClonerEffector/Source/ClonerEffectorEditor/Private/Effector/Customizations/CEEditorEffectorComponentDetailCustomization.h b/Engine/Plugins/Experimental/ClonerEffector/Source/ClonerEffectorEditor/Private/Effector/Customizations/CEEditorEffectorComponentDetailCustomization.h
--- a/Engine/Plugins/Experimental/ClonerEffector/Source/ClonerEffectorEditor/Private/Effector/Customizations/CEEditorEffectorComponentDetailCustomization.h
+++ b/Engine/Plugins/Experimental/ClonerEffector/Source/ClonerEffectorEditor/Private/Effector/Customizations/CEEditorEffectorComponentDetailCustomization.h
@@ -3,6 +3,11 @@
 #pragma once
 
 #include "IDetailCustomization.h"
+#include "Templates/SharedPointerFwd.h"
+
+struct FPropertyChangedEvent;
+class IPropertyHandle;
+class IPropertyUtilities;
 
 /** Used to customize effector component properties in details panel */
 class FCEEditorEffectorComponentDetailCustomization : public IDetailCustomization
@@ -24,4 +29,6 @@ public:
 
 protected:
 	static void RemoveEmptySections();
+	static void OnChildPropertyChanged(const FPropertyChangedEvent& InEvent, TWeakPtr<IPropertyHandle> InParentHandleWeak);
+	static void OnPropertyChanged(const FPropertyChangedEvent& InEvent, TWeakPtr<IPropertyUtilities> InUtilitiesWeak);
 };
