MotionDesign : ClonerEffector
- Fixed Spawn Linked Effector causes cloner to disappear when spawning the first effector

diff --git a/Engine/Plugins/Experimental/ClonerEffector/Source/ClonerEffector/Private/Cloner/Extensions/CEClonerEffectorExtension.cpp b/Engine/Plugins/Experimental/ClonerEffector/Source/ClonerEffector/Private/Cloner/Extensions/CEClonerEffectorExtension.cpp
--- a/Engine/Plugins/Experimental/ClonerEffector/Source/ClonerEffector/Private/Cloner/Extensions/CEClonerEffectorExtension.cpp
+++ b/Engine/Plugins/Experimental/ClonerEffector/Source/ClonerEffector/Private/Cloner/Extensions/CEClonerEffectorExtension.cpp
@@ -116,7 +116,7 @@ void UCEClonerEffectorExtension::OnEffectorsChanged()
 
 	for (TArray<TWeakObjectPtr<AActor>>::TIterator It(EffectorActorsWeak); It; ++It)
 	{
-		if (AActor* EffectorActor = It->Get())
+		if (const AActor* EffectorActor = It->Get())
 		{
 			TArray<UCEEffectorComponent*> EffectorComponents;
 			EffectorActor->GetComponents(EffectorComponents, /** Recurse */false);
@@ -168,7 +168,7 @@ void UCEClonerEffectorExtension::OnEffectorsChanged()
 
 	for (const TWeakObjectPtr<UCEEffectorComponent>& EffectorWeak : SetEffectorsWeak)
 	{
-		if (UCEEffectorComponent* Effector = EffectorWeak.Get())
+		if (const UCEEffectorComponent* Effector = EffectorWeak.Get())
 		{
 			const int32 ChannelIdentifier = Effector->GetChannelIdentifier();
 
@@ -188,14 +188,14 @@ void UCEClonerEffectorExtension::OnEffectorsChanged()
 	}
 
 	// Apply changes
-	EffectorDataInterfaces->Commit();
+	MarkExtensionDirty();
 }
 
 void UCEClonerEffectorExtension::OnEffectorActorsChanged()
 {
 	for (TArray<TWeakObjectPtr<AActor>>::TIterator It(EffectorActorsWeak); It; ++It)
 	{
-		if (AActor* EffectorActor = It->Get())
+		if (const AActor* EffectorActor = It->Get())
 		{
 			if (!EffectorActor->FindComponentByClass<UCEEffectorComponent>())
 			{
@@ -210,7 +210,7 @@ void UCEClonerEffectorExtension::OnEffectorActorsChanged()
 #if WITH_EDITOR
 void UCEClonerEffectorExtension::CreateLinkedEffector()
 {
-	UCEClonerComponent* ClonerComponent = GetClonerComponent();
+	const UCEClonerComponent* ClonerComponent = GetClonerComponent();
 	if (!IsValid(ClonerComponent))
 	{
 		return;
