MotionDesign : ClonerEffector
- Moving grid constraints into an extension
- Constraints properties can now be keyframed
- Allowing extension to filter allowed layouts
- Allowing layout to filter supported extensions

diff --git a/Engine/Plugins/Experimental/ClonerEffector/Source/ClonerEffector/Private/Cloner/CEClonerActor.cpp b/Engine/Plugins/Experimental/ClonerEffector/Source/ClonerEffector/Private/Cloner/CEClonerActor.cpp
--- a/Engine/Plugins/Experimental/ClonerEffector/Source/ClonerEffector/Private/Cloner/CEClonerActor.cpp
+++ b/Engine/Plugins/Experimental/ClonerEffector/Source/ClonerEffector/Private/Cloner/CEClonerActor.cpp
@@ -4,6 +4,7 @@
 
 #include "Cloner/CEClonerComponent.h"
 #include "Cloner/Extensions/CEClonerCollisionExtension.h"
+#include "Cloner/Extensions/CEClonerConstraintExtension.h"
 #include "Cloner/Extensions/CEClonerEffectorExtension.h"
 #include "Cloner/Extensions/CEClonerEmitterSpawnExtension.h"
 #include "Cloner/Extensions/CEClonerLifetimeExtension.h"
@@ -238,14 +239,25 @@ void ACEClonerActor::MigrateDeprecatedProperties()
 					GridLayout->SetSpacingX(PrevGridLayout->GetSpacingX());
 					GridLayout->SetSpacingY(PrevGridLayout->GetSpacingY());
 					GridLayout->SetSpacingZ(PrevGridLayout->GetSpacingZ());
-					GridLayout->SetConstraint(PrevGridLayout->GetConstraint());
-					GridLayout->SetInvertConstraint(PrevGridLayout->GetInvertConstraint());
-					GridLayout->SetSphereConstraint(PrevGridLayout->GetSphereConstraint());
-					GridLayout->SetCylinderConstraint(PrevGridLayout->GetCylinderConstraint());
-					GridLayout->SetTextureConstraint(PrevGridLayout->GetTextureConstraint());
 					GridLayout->SetTwistAxis(PrevGridLayout->GetTwistAxis());
 					GridLayout->SetTwistFactor(PrevGridLayout->GetTwistFactor() * 100);
 				}
+
+				if (UCEClonerConstraintExtension* ConstraintExtension = ClonerComponent->FindOrAddExtension<UCEClonerConstraintExtension>())
+				{
+					ConstraintExtension->SetConstraint(PrevGridLayout->GetConstraint());
+					ConstraintExtension->SetInvertConstraint(PrevGridLayout->GetInvertConstraint());
+					ConstraintExtension->SetSphereRadius(PrevGridLayout->GetSphereConstraint().Radius);
+					ConstraintExtension->SetSphereCenter(PrevGridLayout->GetSphereConstraint().Center);
+					ConstraintExtension->SetCylinderRadius(PrevGridLayout->GetCylinderConstraint().Radius);
+					ConstraintExtension->SetCylinderHeight(PrevGridLayout->GetCylinderConstraint().Height);
+					ConstraintExtension->SetCylinderCenter(PrevGridLayout->GetCylinderConstraint().Center);
+					ConstraintExtension->SetTextureAsset(PrevGridLayout->GetTextureConstraint().Texture.Get());
+					ConstraintExtension->SetTexturePlane(PrevGridLayout->GetTextureConstraint().Plane);
+					ConstraintExtension->SetTextureSampleMode(PrevGridLayout->GetTextureConstraint().Channel);
+					ConstraintExtension->SetTextureCompareMode(PrevGridLayout->GetTextureConstraint().CompareMode);
+					ConstraintExtension->SetTextureThreshold(PrevGridLayout->GetTextureConstraint().Threshold);
+				}
 			}
 			else if (UCEClonerLineLayout* PrevLineLayout = Cast<UCEClonerLineLayout>(LayoutPair.Value))
 			{
diff --git a/Engine/Plugins/Experimental/ClonerEffector/Source/ClonerEffector/Private/Cloner/CEClonerComponent.cpp b/Engine/Plugins/Experimental/ClonerEffector/Source/ClonerEffector/Private/Cloner/CEClonerComponent.cpp
--- a/Engine/Plugins/Experimental/ClonerEffector/Source/ClonerEffector/Private/Cloner/CEClonerComponent.cpp
+++ b/Engine/Plugins/Experimental/ClonerEffector/Source/ClonerEffector/Private/Cloner/CEClonerComponent.cpp
@@ -1711,9 +1711,9 @@ void UCEClonerComponent::OnActiveLayoutChanged()
 	TSet<TObjectPtr<UCEClonerExtensionBase>> PrevActiveExtensions(ActiveExtensions);
 	ActiveExtensions.Empty();
 
-	for (const FName& ExtensionName : Layout->GetSupportedExtensions())
+	for (const TSubclassOf<UCEClonerExtensionBase>& ExtensionClass : Layout->GetSupportedExtensions())
 	{
-		if (UCEClonerExtensionBase* Extension = FindOrAddExtension(ExtensionName))
+		if (UCEClonerExtensionBase* Extension = FindOrAddExtension(ExtensionClass.Get()))
 		{
 			if (!PrevActiveExtensions.Contains(Extension))
 			{
diff --git a/Engine/Plugins/Experimental/ClonerEffector/Source/ClonerEffector/Private/Cloner/Extensions/CEClonerConstraintExtension.cpp b/Engine/Plugins/Experimental/ClonerEffector/Source/ClonerEffector/Private/Cloner/Extensions/CEClonerConstraintExtension.cpp
new file mode 100644
--- /dev/null
+++ b/Engine/Plugins/Experimental/ClonerEffector/Source/ClonerEffector/Private/Cloner/Extensions/CEClonerConstraintExtension.cpp
@@ -0,0 +1,209 @@
+// Copyright Epic Games, Inc. All Rights Reserved.
+
+#include "Cloner/Extensions/CEClonerConstraintExtension.h"
+
+#include "Cloner/CEClonerComponent.h"
+#include "Cloner/Layouts/CEClonerGridLayout.h"
+#include "NiagaraDataInterfaceTexture.h"
+#include "NiagaraSystem.h"
+
+void UCEClonerConstraintExtension::SetConstraint(ECEClonerGridConstraint InConstraint)
+{
+	if (Constraint == InConstraint)
+	{
+		return;
+	}
+
+	Constraint = InConstraint;
+	MarkExtensionDirty();
+}
+
+void UCEClonerConstraintExtension::SetInvertConstraint(bool bInInvertConstraint)
+{
+	if (bInvertConstraint == bInInvertConstraint)
+	{
+		return;
+	}
+
+	bInvertConstraint = bInInvertConstraint;
+	MarkExtensionDirty();
+}
+
+void UCEClonerConstraintExtension::SetSphereRadius(float InSphereRadius)
+{
+	if (FMath::IsNearlyEqual(SphereRadius, InSphereRadius))
+	{
+		return;
+	}
+
+	SphereRadius = InSphereRadius;
+	MarkExtensionDirty();
+}
+
+void UCEClonerConstraintExtension::SetSphereCenter(const FVector& InSphereCenter)
+{
+	if (SphereCenter.Equals(InSphereCenter))
+	{
+		return;
+	}
+
+	SphereCenter = InSphereCenter;
+	MarkExtensionDirty();
+}
+
+void UCEClonerConstraintExtension::SetCylinderRadius(float InCylinderRadius)
+{
+	if (FMath::IsNearlyEqual(CylinderRadius, InCylinderRadius))
+	{
+		return;
+	}
+
+	CylinderRadius = InCylinderRadius;
+	MarkExtensionDirty();
+}
+
+void UCEClonerConstraintExtension::SetCylinderHeight(float InCylinderHeight)
+{
+	if (FMath::IsNearlyEqual(CylinderHeight, InCylinderHeight))
+	{
+		return;
+	}
+
+	CylinderHeight = InCylinderHeight;
+	MarkExtensionDirty();
+}
+
+void UCEClonerConstraintExtension::SetCylinderCenter(const FVector& InCylinderCenter)
+{
+	if (CylinderCenter.Equals(InCylinderCenter))
+	{
+		return;
+	}
+
+	CylinderCenter = InCylinderCenter;
+	MarkExtensionDirty();
+}
+
+void UCEClonerConstraintExtension::SetTextureAsset(UTexture* InTexture)
+{
+	if (TextureAsset == InTexture)
+	{
+		return;
+	}
+
+	TextureAsset = InTexture;
+	MarkExtensionDirty();
+}
+
+void UCEClonerConstraintExtension::SetTextureSampleMode(ECEClonerTextureSampleChannel InSampleMode)
+{
+	if (TextureSampleMode == InSampleMode)
+	{
+		return;
+	}
+
+	TextureSampleMode = InSampleMode;
+	MarkExtensionDirty();
+}
+
+void UCEClonerConstraintExtension::SetTexturePlane(ECEClonerPlane InPlane)
+{
+	if (TexturePlane == InPlane)
+	{
+		return;
+	}
+
+	TexturePlane = InPlane;
+	MarkExtensionDirty();
+}
+
+void UCEClonerConstraintExtension::SetTextureCompareMode(ECEClonerCompareMode InMode)
+{
+	if (TextureCompareMode == InMode)
+	{
+		return;
+	}
+
+	TextureCompareMode = InMode;
+	MarkExtensionDirty();
+}
+
+void UCEClonerConstraintExtension::SetTextureThreshold(float InThreshold)
+{
+	if (FMath::IsNearlyEqual(TextureThreshold, InThreshold))
+	{
+		return;
+	}
+
+	TextureThreshold = InThreshold;
+	MarkExtensionDirty();
+}
+
+#if WITH_EDITOR
+const TCEPropertyChangeDispatcher<UCEClonerConstraintExtension> UCEClonerConstraintExtension::PropertyChangeDispatcher =
+{
+	{ GET_MEMBER_NAME_CHECKED(UCEClonerConstraintExtension, Constraint), &UCEClonerConstraintExtension::OnExtensionPropertyChanged },
+	{ GET_MEMBER_NAME_CHECKED(UCEClonerConstraintExtension, bInvertConstraint), &UCEClonerConstraintExtension::OnExtensionPropertyChanged },
+	{ GET_MEMBER_NAME_CHECKED(UCEClonerConstraintExtension, SphereRadius), &UCEClonerConstraintExtension::OnExtensionPropertyChanged },
+	{ GET_MEMBER_NAME_CHECKED(UCEClonerConstraintExtension, SphereCenter), &UCEClonerConstraintExtension::OnExtensionPropertyChanged },
+	{ GET_MEMBER_NAME_CHECKED(UCEClonerConstraintExtension, CylinderRadius), &UCEClonerConstraintExtension::OnExtensionPropertyChanged },
+	{ GET_MEMBER_NAME_CHECKED(UCEClonerConstraintExtension, CylinderHeight), &UCEClonerConstraintExtension::OnExtensionPropertyChanged },
+	{ GET_MEMBER_NAME_CHECKED(UCEClonerConstraintExtension, CylinderCenter), &UCEClonerConstraintExtension::OnExtensionPropertyChanged },
+	{ GET_MEMBER_NAME_CHECKED(UCEClonerConstraintExtension, TextureAsset), &UCEClonerConstraintExtension::OnExtensionPropertyChanged },
+	{ GET_MEMBER_NAME_CHECKED(UCEClonerConstraintExtension, TextureSampleMode), &UCEClonerConstraintExtension::OnExtensionPropertyChanged },
+	{ GET_MEMBER_NAME_CHECKED(UCEClonerConstraintExtension, TexturePlane), &UCEClonerConstraintExtension::OnExtensionPropertyChanged },
+	{ GET_MEMBER_NAME_CHECKED(UCEClonerConstraintExtension, TextureCompareMode), &UCEClonerConstraintExtension::OnExtensionPropertyChanged },
+	{ GET_MEMBER_NAME_CHECKED(UCEClonerConstraintExtension, TextureThreshold), &UCEClonerConstraintExtension::OnExtensionPropertyChanged },
+};
+
+void UCEClonerConstraintExtension::PostEditChangeProperty(FPropertyChangedEvent& InPropertyChangedEvent)
+{
+	Super::PostEditChangeProperty(InPropertyChangedEvent);
+
+	PropertyChangeDispatcher.OnPropertyChanged(this, InPropertyChangedEvent);
+}
+#endif
+
+void UCEClonerConstraintExtension::OnExtensionParametersChanged(UCEClonerComponent* InComponent)
+{
+	Super::OnExtensionParametersChanged(InComponent);
+
+	FNiagaraUserRedirectionParameterStore& ExposedParameters = InComponent->GetAsset()->GetExposedParameters();
+	static const FNiagaraVariable ConstraintVar(FNiagaraTypeDefinition(StaticEnum<ECEClonerGridConstraint>()), TEXT("Constraint"));
+	ExposedParameters.SetParameterValue<int32>(static_cast<int32>(Constraint), ConstraintVar);
+
+	InComponent->SetBoolParameter(TEXT("ConstraintInvert"), Constraint != ECEClonerGridConstraint::None ? bInvertConstraint : false);
+
+	// Sphere
+	InComponent->SetVectorParameter(TEXT("ConstraintSphereCenter"), SphereCenter);
+
+	InComponent->SetFloatParameter(TEXT("ConstraintSphereRadius"), SphereRadius);
+
+	// Cylinder
+	InComponent->SetVectorParameter(TEXT("ConstraintCylinderCenter"), CylinderCenter);
+
+	InComponent->SetFloatParameter(TEXT("ConstraintCylinderHeight"), CylinderHeight);
+
+	InComponent->SetFloatParameter(TEXT("ConstraintCylinderRadius"), CylinderRadius);
+
+	// Texture
+	static const FNiagaraVariable ConstraintTextureSamplerVar(FNiagaraTypeDefinition(UNiagaraDataInterfaceTexture::StaticClass()), TEXT("ConstraintTextureSampler"));
+	UNiagaraDataInterfaceTexture* TextureSamplerDI = Cast<UNiagaraDataInterfaceTexture>(ExposedParameters.GetDataInterface(ConstraintTextureSamplerVar));
+	TextureSamplerDI->SetTexture(TextureAsset.Get());
+
+	static const FNiagaraVariable ConstraintTexturePlaneVar(FNiagaraTypeDefinition(StaticEnum<ECEClonerPlane>()), TEXT("ConstraintTexturePlane"));
+	ExposedParameters.SetParameterValue<int32>(static_cast<int32>(TexturePlane), ConstraintTexturePlaneVar);
+
+	static const FNiagaraVariable ConstraintTextureChannelVar(FNiagaraTypeDefinition(StaticEnum<ECEClonerTextureSampleChannel>()), TEXT("ConstraintTextureChannel"));
+	ExposedParameters.SetParameterValue<int32>(static_cast<int32>(TextureSampleMode), ConstraintTextureChannelVar);
+
+	static const FNiagaraVariable ConstraintTextureCompareModeVar(FNiagaraTypeDefinition(StaticEnum<ECEClonerCompareMode>()), TEXT("ConstraintTextureCompareMode"));
+	ExposedParameters.SetParameterValue<int32>(static_cast<int32>(TextureCompareMode), ConstraintTextureCompareModeVar);
+
+	InComponent->SetFloatParameter(TEXT("ConstraintTextureThreshold"), TextureThreshold);
+}
+
+bool UCEClonerConstraintExtension::IsLayoutSupported(const UCEClonerLayoutBase* InLayout) const
+{
+	return InLayout->IsA<UCEClonerGridLayout>();
+}
diff --git a/Engine/Plugins/Experimental/ClonerEffector/Source/ClonerEffector/Private/Cloner/Layouts/CEClonerGridLayout.cpp b/Engine/Plugins/Experimental/ClonerEffector/Source/ClonerEffector/Private/Cloner/Layouts/CEClonerGridLayout.cpp
--- a/Engine/Plugins/Experimental/ClonerEffector/Source/ClonerEffector/Private/Cloner/Layouts/CEClonerGridLayout.cpp
+++ b/Engine/Plugins/Experimental/ClonerEffector/Source/ClonerEffector/Private/Cloner/Layouts/CEClonerGridLayout.cpp
@@ -3,7 +3,6 @@
 #include "Cloner/Layouts/CEClonerGridLayout.h"
 
 #include "Cloner/CEClonerComponent.h"
-#include "NiagaraDataInterfaceTexture.h"
 #include "NiagaraSystem.h"
 
 void UCEClonerGridLayout::SetCountX(int32 InCountX)
@@ -72,46 +71,6 @@ void UCEClonerGridLayout::SetSpacingZ(float InSpacingZ)
 	MarkLayoutDirty();
 }
 
-void UCEClonerGridLayout::SetConstraint(ECEClonerGridConstraint InConstraint)
-{
-	if (Constraint == InConstraint)
-	{
-		return;
-	}
-
-	Constraint = InConstraint;
-	MarkLayoutDirty();
-}
-
-void UCEClonerGridLayout::SetInvertConstraint(bool bInInvertConstraint)
-{
-	if (bInvertConstraint == bInInvertConstraint)
-	{
-		return;
-	}
-
-	bInvertConstraint = bInInvertConstraint;
-	MarkLayoutDirty();
-}
-
-void UCEClonerGridLayout::SetSphereConstraint(const FCEClonerGridConstraintSphere& InConstraint)
-{
-	SphereConstraint = InConstraint;
-	MarkLayoutDirty();
-}
-
-void UCEClonerGridLayout::SetCylinderConstraint(const FCEClonerGridConstraintCylinder& InConstraint)
-{
-	CylinderConstraint = InConstraint;
-	MarkLayoutDirty();
-}
-
-void UCEClonerGridLayout::SetTextureConstraint(const FCEClonerGridConstraintTexture& InConstraint)
-{
-	TextureConstraint = InConstraint;
-	MarkLayoutDirty();
-}
-
 void UCEClonerGridLayout::SetTwistFactor(float InFactor)
 {
 	if (TwistFactor == InFactor)
@@ -143,11 +102,6 @@ const TCEPropertyChangeDispatcher<UCEClonerGridLayout> UCEClonerGridLayout::Prop
 	{ GET_MEMBER_NAME_CHECKED(UCEClonerGridLayout, SpacingX), &UCEClonerGridLayout::OnLayoutPropertyChanged },
 	{ GET_MEMBER_NAME_CHECKED(UCEClonerGridLayout, SpacingY), &UCEClonerGridLayout::OnLayoutPropertyChanged },
 	{ GET_MEMBER_NAME_CHECKED(UCEClonerGridLayout, SpacingZ), &UCEClonerGridLayout::OnLayoutPropertyChanged },
-	{ GET_MEMBER_NAME_CHECKED(UCEClonerGridLayout, Constraint), &UCEClonerGridLayout::OnLayoutPropertyChanged },
-	{ GET_MEMBER_NAME_CHECKED(UCEClonerGridLayout, bInvertConstraint), &UCEClonerGridLayout::OnLayoutPropertyChanged },
-	{ GET_MEMBER_NAME_CHECKED(UCEClonerGridLayout, SphereConstraint), &UCEClonerGridLayout::OnLayoutPropertyChanged },
-	{ GET_MEMBER_NAME_CHECKED(UCEClonerGridLayout, CylinderConstraint), &UCEClonerGridLayout::OnLayoutPropertyChanged },
-	{ GET_MEMBER_NAME_CHECKED(UCEClonerGridLayout, TextureConstraint), &UCEClonerGridLayout::OnLayoutPropertyChanged },
 	{ GET_MEMBER_NAME_CHECKED(UCEClonerGridLayout, TwistFactor), &UCEClonerGridLayout::OnLayoutPropertyChanged },
 	{ GET_MEMBER_NAME_CHECKED(UCEClonerGridLayout, TwistAxis), &UCEClonerGridLayout::OnLayoutPropertyChanged },
 };
@@ -172,39 +126,9 @@ void UCEClonerGridLayout::OnLayoutParametersChanged(UCEClonerComponent* InCompon
 
 	InComponent->SetVectorParameter(TEXT("GridSpacing"), FVector(SpacingX, SpacingY, SpacingZ));
 
-	FNiagaraUserRedirectionParameterStore& ExposedParameters = InComponent->GetAsset()->GetExposedParameters();
-	static const FNiagaraVariable ConstraintVar(FNiagaraTypeDefinition(StaticEnum<ECEClonerGridConstraint>()), TEXT("Constraint"));
-	ExposedParameters.SetParameterValue<int32>(static_cast<int32>(Constraint), ConstraintVar);
-
-	InComponent->SetBoolParameter(TEXT("ConstraintInvert"), Constraint != ECEClonerGridConstraint::None ? bInvertConstraint : false);
-
-	InComponent->SetVectorParameter(TEXT("ConstraintCylinderCenter"), CylinderConstraint.Center);
-
-	InComponent->SetFloatParameter(TEXT("ConstraintCylinderHeight"), CylinderConstraint.Height);
-
-	InComponent->SetFloatParameter(TEXT("ConstraintCylinderRadius"), CylinderConstraint.Radius);
-
-	InComponent->SetVectorParameter(TEXT("ConstraintSphereCenter"), SphereConstraint.Center);
-
-	InComponent->SetFloatParameter(TEXT("ConstraintSphereRadius"), SphereConstraint.Radius);
-
-	static const FNiagaraVariable ConstraintTextureChannelVar(FNiagaraTypeDefinition(StaticEnum<ECEClonerTextureSampleChannel>()), TEXT("ConstraintTextureChannel"));
-	ExposedParameters.SetParameterValue<int32>(static_cast<int32>(TextureConstraint.Channel), ConstraintTextureChannelVar);
-
-	static const FNiagaraVariable ConstraintTextureCompareModeVar(FNiagaraTypeDefinition(StaticEnum<ECEClonerCompareMode>()), TEXT("ConstraintTextureCompareMode"));
-	ExposedParameters.SetParameterValue<int32>(static_cast<int32>(TextureConstraint.CompareMode), ConstraintTextureCompareModeVar);
-
-	static const FNiagaraVariable ConstraintTexturePlaneVar(FNiagaraTypeDefinition(StaticEnum<ECEClonerPlane>()), TEXT("ConstraintTexturePlane"));
-	ExposedParameters.SetParameterValue<int32>(static_cast<int32>(TextureConstraint.Plane), ConstraintTexturePlaneVar);
-
-	InComponent->SetFloatParameter(TEXT("ConstraintTextureThreshold"), TextureConstraint.Threshold);
-
-	static const FNiagaraVariable ConstraintTextureSamplerVar(FNiagaraTypeDefinition(UNiagaraDataInterfaceTexture::StaticClass()), TEXT("ConstraintTextureSampler"));
-	UNiagaraDataInterfaceTexture* TextureSamplerDI = Cast<UNiagaraDataInterfaceTexture>(ExposedParameters.GetDataInterface(ConstraintTextureSamplerVar));
-	TextureSamplerDI->SetTexture(TextureConstraint.Texture.Get());
-
 	InComponent->SetFloatParameter(TEXT("TwistFactor"), TwistFactor / 100.f);
 
+	FNiagaraUserRedirectionParameterStore& ExposedParameters = InComponent->GetAsset()->GetExposedParameters();
 	static const FNiagaraVariable TwistAxisVar(FNiagaraTypeDefinition(StaticEnum<ENiagaraOrientationAxis>()), TEXT("TwistAxis"));
 	ExposedParameters.SetParameterValue<int32>(static_cast<int32>(TwistAxis), TwistAxisVar);
 }
diff --git a/Engine/Plugins/Experimental/ClonerEffector/Source/ClonerEffector/Private/Cloner/Layouts/CEClonerLayoutBase.cpp b/Engine/Plugins/Experimental/ClonerEffector/Source/ClonerEffector/Private/Cloner/Layouts/CEClonerLayoutBase.cpp
--- a/Engine/Plugins/Experimental/ClonerEffector/Source/ClonerEffector/Private/Cloner/Layouts/CEClonerLayoutBase.cpp
+++ b/Engine/Plugins/Experimental/ClonerEffector/Source/ClonerEffector/Private/Cloner/Layouts/CEClonerLayoutBase.cpp
@@ -4,6 +4,7 @@
 
 #include "Cloner/CEClonerActor.h"
 #include "Cloner/CEClonerComponent.h"
+#include "Cloner/Extensions/CEClonerExtensionBase.h"
 #include "Misc/PackageName.h"
 #include "NiagaraEmitter.h"
 #include "NiagaraMeshRendererProperties.h"
@@ -222,16 +223,38 @@ bool UCEClonerLayoutBase::CopyTo(UCEClonerLayoutBase* InOtherLayout) const
 	return true;
 }
 
-TSet<FName> UCEClonerLayoutBase::GetSupportedExtensions() const
+TSet<TSubclassOf<UCEClonerExtensionBase>> UCEClonerLayoutBase::GetSupportedExtensions() const
 {
-	TSet<FName> ModuleSupported;
+	TSet<TSubclassOf<UCEClonerExtensionBase>> ExtensionSupported;
 
-	if (UCEClonerSubsystem* ClonerSubsystem = UCEClonerSubsystem::Get())
+	if (const UCEClonerSubsystem* ClonerSubsystem = UCEClonerSubsystem::Get())
 	{
-		ModuleSupported.Append(ClonerSubsystem->GetExtensionNames());
+		for (const TSubclassOf<UCEClonerExtensionBase>& ExtensionClass : ClonerSubsystem->GetExtensionClasses())
+		{
+			const UCEClonerExtensionBase* Extension = ExtensionClass.GetDefaultObject();
+
+			if (!Extension)
+			{
+				continue;
+			}
+
+			// Does the layout supports this extension
+			if (!IsExtensionSupported(Extension))
+			{
+				continue;
+			}
+
+			// Does the extension supports this layout
+			if (!Extension->IsLayoutSupported(this))
+			{
+				continue;
+			}
+
+			ExtensionSupported.Add(ExtensionClass);
+		}
 	}
 
-	return ModuleSupported;
+	return ExtensionSupported;
 }
 
 bool UCEClonerLayoutBase::IsLayoutDirty() const
diff --git a/Engine/Plugins/Experimental/ClonerEffector/Source/ClonerEffector/Private/Subsystems/CEClonerSubsystem.cpp b/Engine/Plugins/Experimental/ClonerEffector/Source/ClonerEffector/Private/Subsystems/CEClonerSubsystem.cpp
--- a/Engine/Plugins/Experimental/ClonerEffector/Source/ClonerEffector/Private/Subsystems/CEClonerSubsystem.cpp
+++ b/Engine/Plugins/Experimental/ClonerEffector/Source/ClonerEffector/Private/Subsystems/CEClonerSubsystem.cpp
@@ -216,11 +216,18 @@ bool UCEClonerSubsystem::IsExtensionClassRegistered(UClass* InClass) const
 	return !!ExtensionClasses.FindKey(ExtensionClass);
 }
 
-TArray<FName> UCEClonerSubsystem::GetExtensionNames() const
+TSet<FName> UCEClonerSubsystem::GetExtensionNames() const
 {
 	TArray<FName> ExtensionNames;
 	ExtensionClasses.GenerateKeyArray(ExtensionNames);
-	return ExtensionNames;
+	return TSet<FName>(ExtensionNames);
+}
+
+TSet<TSubclassOf<UCEClonerExtensionBase>> UCEClonerSubsystem::GetExtensionClasses() const
+{
+	TArray<TSubclassOf<UCEClonerExtensionBase>> Extensions;
+	ExtensionClasses.GenerateValueArray(Extensions);
+	return TSet<TSubclassOf<UCEClonerExtensionBase>>(Extensions);
 }
 
 FName UCEClonerSubsystem::FindExtensionName(TSubclassOf<UCEClonerExtensionBase> InClass) const
diff --git a/Engine/Plugins/Experimental/ClonerEffector/Source/ClonerEffector/Public/Cloner/Extensions/CEClonerConstraintExtension.h b/Engine/Plugins/Experimental/ClonerEffector/Source/ClonerEffector/Public/Cloner/Extensions/CEClonerConstraintExtension.h
new file mode 100644
--- /dev/null
+++ b/Engine/Plugins/Experimental/ClonerEffector/Source/ClonerEffector/Public/Cloner/Extensions/CEClonerConstraintExtension.h
@@ -0,0 +1,190 @@
+// Copyright Epic Games, Inc. All Rights Reserved.
+
+#pragma once
+
+#include "CEClonerExtensionBase.h"
+#include "CEPropertyChangeDispatcher.h"
+#include "CEClonerConstraintExtension.generated.h"
+
+/** Extension dealing with clone grid constraints */
+UCLASS(MinimalAPI, BlueprintType, Within=CEClonerComponent, meta=(Section="Cloner", Priority=0))
+class UCEClonerConstraintExtension : public UCEClonerExtensionBase
+{
+	GENERATED_BODY()
+
+public:
+	UCEClonerConstraintExtension()
+		: UCEClonerExtensionBase(
+			TEXT("Constraint")
+			, 0
+		)
+	{}
+
+	UFUNCTION(BlueprintCallable, Category="Cloner")
+	CLONEREFFECTOR_API void SetConstraint(ECEClonerGridConstraint InConstraint);
+
+	UFUNCTION(BlueprintPure, Category="Cloner")
+	ECEClonerGridConstraint GetConstraint() const
+	{
+		return Constraint;
+	}
+
+	UFUNCTION(BlueprintCallable, Category="Cloner")
+	CLONEREFFECTOR_API void SetInvertConstraint(bool bInInvertConstraint);
+
+	UFUNCTION(BlueprintPure, Category="Cloner")
+	bool GetInvertConstraint() const
+	{
+		return bInvertConstraint;
+	}
+
+	UFUNCTION(BlueprintCallable, Category="Cloner")
+	CLONEREFFECTOR_API void SetSphereRadius(float InSphereRadius);
+
+	UFUNCTION(BlueprintPure, Category="Cloner")
+	float GetSphereRadius() const
+	{
+		return SphereRadius;
+	}
+
+	UFUNCTION(BlueprintCallable, Category="Cloner")
+	CLONEREFFECTOR_API void SetSphereCenter(const FVector& InSphereCenter);
+
+	UFUNCTION(BlueprintPure, Category="Cloner")
+	FVector GetSphereCenter() const
+	{
+		return SphereCenter;
+	}
+
+	UFUNCTION(BlueprintCallable, Category="Cloner")
+	CLONEREFFECTOR_API void SetCylinderRadius(float InCylinderRadius);
+
+	UFUNCTION(BlueprintPure, Category="Cloner")
+	float GetCylinderRadius() const
+	{
+		return CylinderRadius;
+	}
+
+	UFUNCTION(BlueprintCallable, Category="Cloner")
+	CLONEREFFECTOR_API void SetCylinderHeight(float InCylinderHeight);
+
+	UFUNCTION(BlueprintPure, Category="Cloner")
+	float GetCylinderHeight() const
+	{
+		return CylinderHeight;
+	}
+
+	UFUNCTION(BlueprintCallable, Category="Cloner")
+	CLONEREFFECTOR_API void SetCylinderCenter(const FVector& InCylinderCenter);
+
+	UFUNCTION(BlueprintPure, Category="Cloner")
+	FVector GetCylinderCenter() const
+	{
+		return CylinderCenter;
+	}
+
+	UFUNCTION(BlueprintCallable, Category="Cloner")
+	CLONEREFFECTOR_API void SetTextureAsset(UTexture* InTexture);
+
+	UFUNCTION(BlueprintPure, Category="Cloner")
+	UTexture* GetTextureAsset() const
+	{
+		return TextureAsset;
+	}
+
+	UFUNCTION(BlueprintCallable, Category="Cloner")
+	CLONEREFFECTOR_API void SetTextureSampleMode(ECEClonerTextureSampleChannel InSampleMode);
+
+	UFUNCTION(BlueprintPure, Category="Cloner")
+	ECEClonerTextureSampleChannel GetTextureSampleMode() const
+	{
+		return TextureSampleMode;
+	}
+
+	UFUNCTION(BlueprintCallable, Category="Cloner")
+	CLONEREFFECTOR_API void SetTexturePlane(ECEClonerPlane InPlane);
+
+	UFUNCTION(BlueprintPure, Category="Cloner")
+	ECEClonerPlane GetTexturePlane() const
+	{
+		return TexturePlane;
+	}
+
+	UFUNCTION(BlueprintCallable, Category="Cloner")
+	CLONEREFFECTOR_API void SetTextureCompareMode(ECEClonerCompareMode InMode);
+
+	UFUNCTION(BlueprintPure, Category="Cloner")
+	ECEClonerCompareMode GetTextureCompareMode() const
+	{
+		return TextureCompareMode;
+	}
+
+	UFUNCTION(BlueprintCallable, Category="Cloner")
+	CLONEREFFECTOR_API void SetTextureThreshold(float InThreshold);
+
+	UFUNCTION(BlueprintPure, Category="Cloner")
+	float GetTextureThreshold() const
+	{
+		return TextureThreshold;
+	}
+
+protected:
+	//~ Begin UObject
+#if WITH_EDITOR
+	virtual void PostEditChangeProperty(FPropertyChangedEvent& InPropertyChangedEvent) override;
+#endif
+	//~ End UObject
+
+	//~ Begin UCEClonerExtensionBase
+	virtual void OnExtensionParametersChanged(UCEClonerComponent* InComponent) override;
+	virtual bool IsLayoutSupported(const UCEClonerLayoutBase* InLayout) const override;
+	//~ End UCEClonerExtensionBase
+
+	UPROPERTY(EditInstanceOnly, BlueprintReadWrite, Setter, Getter, Category="Constraint")
+	ECEClonerGridConstraint Constraint = ECEClonerGridConstraint::None;
+
+	UPROPERTY(EditInstanceOnly, BlueprintReadWrite, Setter="SetInvertConstraint", Getter="GetInvertConstraint", Category="Constraint", meta=(EditCondition="Constraint != ECEClonerGridConstraint::None", EditConditionHides))
+	bool bInvertConstraint = false;
+
+	/** Sphere constraint */
+
+	UPROPERTY(EditInstanceOnly, BlueprintReadWrite, Setter, Getter, Category="Constraint", meta=(EditCondition="Constraint == ECEClonerGridConstraint::Sphere", EditConditionHides))
+	float SphereRadius = 400.f;
+
+	UPROPERTY(EditInstanceOnly, BlueprintReadWrite, Setter, Getter, Category="Constraint", meta=(EditCondition="Constraint == ECEClonerGridConstraint::Sphere", EditConditionHides))
+	FVector SphereCenter = FVector::ZeroVector;
+
+	/** Cylinder constraint */
+
+	UPROPERTY(EditInstanceOnly, BlueprintReadWrite, Setter, Getter, Category="Constraint", meta=(EditCondition="Constraint == ECEClonerGridConstraint::Cylinder", EditConditionHides))
+	float CylinderRadius = 400.f;
+
+	UPROPERTY(EditInstanceOnly, BlueprintReadWrite, Setter, Getter, Category="Constraint", meta=(EditCondition="Constraint == ECEClonerGridConstraint::Cylinder", EditConditionHides))
+	float CylinderHeight = 800.f;
+
+	UPROPERTY(EditInstanceOnly, BlueprintReadWrite, Setter, Getter, Category="Constraint", meta=(EditCondition="Constraint == ECEClonerGridConstraint::Cylinder", EditConditionHides))
+	FVector CylinderCenter = FVector::ZeroVector;
+
+	/** Texture constraint */
+
+	UPROPERTY(EditInstanceOnly, BlueprintReadWrite, Setter, Getter, Category="Constraint", meta=(EditCondition="Constraint == ECEClonerGridConstraint::Texture", EditConditionHides))
+	TObjectPtr<UTexture> TextureAsset;
+
+	UPROPERTY(EditInstanceOnly, BlueprintReadWrite, Setter, Getter, Category="Constraint", meta=(InvalidEnumValues="Custom", EditCondition="Constraint == ECEClonerGridConstraint::Texture", EditConditionHides))
+	ECEClonerPlane TexturePlane = ECEClonerPlane::XY;
+
+	UPROPERTY(EditInstanceOnly, BlueprintReadWrite, Setter, Getter, Category="Constraint", meta=(EditCondition="Constraint == ECEClonerGridConstraint::Texture", EditConditionHides))
+	ECEClonerTextureSampleChannel TextureSampleMode = ECEClonerTextureSampleChannel::RGBLuminance;
+
+	UPROPERTY(EditInstanceOnly, BlueprintReadWrite, Setter, Getter, Category="Constraint", meta=(EditCondition="Constraint == ECEClonerGridConstraint::Texture", EditConditionHides))
+	ECEClonerCompareMode TextureCompareMode = ECEClonerCompareMode::Greater;
+
+	UPROPERTY(EditInstanceOnly, BlueprintReadWrite, Setter, Getter, Category="Constraint", meta=(EditCondition="Constraint == ECEClonerGridConstraint::Texture", EditConditionHides))
+	float TextureThreshold = 0.f;
+
+private:
+#if WITH_EDITOR
+	/** Used for PECP */
+	static const TCEPropertyChangeDispatcher<UCEClonerConstraintExtension> PropertyChangeDispatcher;
+#endif
+};
\ No newline at end of file
diff --git a/Engine/Plugins/Experimental/ClonerEffector/Source/ClonerEffector/Public/Cloner/Extensions/CEClonerExtensionBase.h b/Engine/Plugins/Experimental/ClonerEffector/Source/ClonerEffector/Public/Cloner/Extensions/CEClonerExtensionBase.h
--- a/Engine/Plugins/Experimental/ClonerEffector/Source/ClonerEffector/Public/Cloner/Extensions/CEClonerExtensionBase.h
+++ b/Engine/Plugins/Experimental/ClonerEffector/Source/ClonerEffector/Public/Cloner/Extensions/CEClonerExtensionBase.h
@@ -73,8 +73,15 @@ public:
 	/** Request refresh extension next tick */
 	void MarkExtensionDirty(bool bInUpdateCloner = true);
 
+	/** Is this extension dirty */
 	bool IsExtensionDirty() const;
 
+	/** Filter supported layout for this extension */
+	virtual bool IsLayoutSupported(const UCEClonerLayoutBase* InLayout) const
+	{
+		return true;
+	}
+
 protected:
 	//~ Begin UObject
 	virtual void PostEditImport() override;
diff --git a/Engine/Plugins/Experimental/ClonerEffector/Source/ClonerEffector/Public/Cloner/Layouts/CEClonerGridLayout.h b/Engine/Plugins/Experimental/ClonerEffector/Source/ClonerEffector/Public/Cloner/Layouts/CEClonerGridLayout.h
--- a/Engine/Plugins/Experimental/ClonerEffector/Source/ClonerEffector/Public/Cloner/Layouts/CEClonerGridLayout.h
+++ b/Engine/Plugins/Experimental/ClonerEffector/Source/ClonerEffector/Public/Cloner/Layouts/CEClonerGridLayout.h
@@ -77,67 +77,51 @@ public:
 	}
 
 	UFUNCTION(BlueprintCallable, Category="Cloner|Layout|Grid")
-	CLONEREFFECTOR_API void SetConstraint(ECEClonerGridConstraint InConstraint);
+	CLONEREFFECTOR_API void SetTwistFactor(float InFactor);
 
 	UFUNCTION(BlueprintPure, Category="Cloner|Layout|Grid")
-	ECEClonerGridConstraint GetConstraint() const
+	float GetTwistFactor() const
 	{
-		return Constraint;
+		return TwistFactor;
 	}
 
 	UFUNCTION(BlueprintCallable, Category="Cloner|Layout|Grid")
-	CLONEREFFECTOR_API void SetInvertConstraint(bool bInInvertConstraint);
+	CLONEREFFECTOR_API void SetTwistAxis(ECEClonerAxis InAxis);
 
 	UFUNCTION(BlueprintPure, Category="Cloner|Layout|Grid")
+	ECEClonerAxis GetTwistAxis() const
+	{
+		return TwistAxis;
+	}
+
+	PRAGMA_DISABLE_DEPRECATION_WARNINGS
+
+	ECEClonerGridConstraint GetConstraint() const
+	{
+		return Constraint;
+	}
+
 	bool GetInvertConstraint() const
 	{
 		return bInvertConstraint;
 	}
 
-	UFUNCTION(BlueprintCallable, Category="Cloner|Layout|Grid")
-	CLONEREFFECTOR_API void SetSphereConstraint(const FCEClonerGridConstraintSphere& InConstraint);
-
-	UFUNCTION(BlueprintPure, Category="Cloner|Layout|Grid")
 	const FCEClonerGridConstraintSphere& GetSphereConstraint() const
 	{
 		return SphereConstraint;
 	}
 
-	UFUNCTION(BlueprintCallable, Category="Cloner|Layout|Grid")
-	CLONEREFFECTOR_API void SetCylinderConstraint(const FCEClonerGridConstraintCylinder& InConstraint);
-
-	UFUNCTION(BlueprintPure, Category="Cloner|Layout|Grid")
 	const FCEClonerGridConstraintCylinder& GetCylinderConstraint() const
 	{
 		return CylinderConstraint;
 	}
 
-	UFUNCTION(BlueprintCallable, Category="Cloner|Layout|Grid")
-	CLONEREFFECTOR_API void SetTextureConstraint(const FCEClonerGridConstraintTexture& InConstraint);
-
-	UFUNCTION(BlueprintPure, Category="Cloner|Layout|Grid")
 	const FCEClonerGridConstraintTexture& GetTextureConstraint() const
 	{
 		return TextureConstraint;
 	}
 
-	UFUNCTION(BlueprintCallable, Category="Cloner|Layout|Grid")
-	CLONEREFFECTOR_API void SetTwistFactor(float InFactor);
-
-	UFUNCTION(BlueprintPure, Category="Cloner|Layout|Grid")
-	float GetTwistFactor() const
-	{
-		return TwistFactor;
-	}
-
-	UFUNCTION(BlueprintCallable, Category="Cloner|Layout|Grid")
-	CLONEREFFECTOR_API void SetTwistAxis(ECEClonerAxis InAxis);
-
-	UFUNCTION(BlueprintPure, Category="Cloner|Layout|Grid")
-	ECEClonerAxis GetTwistAxis() const
-	{
-		return TwistAxis;
-	}
+	PRAGMA_ENABLE_DEPRECATION_WARNINGS
 
 protected:
 	//~ Begin UObject
@@ -168,28 +152,33 @@ protected:
 	UPROPERTY(EditInstanceOnly, BlueprintReadWrite, Setter="SetSpacingZ", Getter="GetSpacingZ", Category="Layout")
 	float SpacingZ = 105.f;
 
-	UPROPERTY(EditInstanceOnly, BlueprintReadWrite, Setter="SetConstraint", Getter="GetConstraint", Category="Layout")
+	UPROPERTY(EditInstanceOnly, BlueprintReadWrite, Setter, Getter, Category="Layout", meta=(UIMin="0", UIMax="100"))
+	float TwistFactor = 0.f;
+
+	UPROPERTY(EditInstanceOnly, BlueprintReadWrite, Setter, Getter, Category="Layout", meta=(InvalidEnumValues="Custom"))
+	ECEClonerAxis TwistAxis = ECEClonerAxis::Y;
+
+private:
+	UE_DEPRECATED(5.5, "Moved to constraint extension")
+	UPROPERTY()
 	ECEClonerGridConstraint Constraint = ECEClonerGridConstraint::None;
 
-	UPROPERTY(EditInstanceOnly, BlueprintReadWrite, Setter="SetInvertConstraint", Getter="GetInvertConstraint", Category="Layout", meta=(EditCondition="Constraint != ECEClonerGridConstraint::None", EditConditionHides))
+	UE_DEPRECATED(5.5, "Moved to constraint extension")
+	UPROPERTY()
 	bool bInvertConstraint = false;
 
-	UPROPERTY(EditInstanceOnly, BlueprintReadWrite, Setter="SetSphereConstraint", Getter="GetSphereConstraint", Category="Layout", meta=(EditCondition="Constraint == ECEClonerGridConstraint::Sphere", EditConditionHides))
+	UE_DEPRECATED(5.5, "Moved to constraint extension")
+	UPROPERTY()
 	FCEClonerGridConstraintSphere SphereConstraint;
 
-	UPROPERTY(EditInstanceOnly, BlueprintReadWrite, Setter="SetCylinderConstraint", Getter="GetCylinderConstraint", Category="Layout", meta=(EditCondition="Constraint == ECEClonerGridConstraint::Cylinder", EditConditionHides))
+	UE_DEPRECATED(5.5, "Moved to constraint extension")
+	UPROPERTY()
 	FCEClonerGridConstraintCylinder CylinderConstraint;
 
-	UPROPERTY(EditInstanceOnly, BlueprintReadWrite, Setter="SetTextureConstraint", Getter="GetTextureConstraint", Category="Layout", meta=(EditCondition="Constraint == ECEClonerGridConstraint::Texture", EditConditionHides))
+	UE_DEPRECATED(5.5, "Moved to constraint extension")
+	UPROPERTY()
 	FCEClonerGridConstraintTexture TextureConstraint;
 
-	UPROPERTY(EditInstanceOnly, BlueprintReadWrite, Setter, Getter, Category="Layout", meta=(UIMin="0", UIMax="100"))
-	float TwistFactor = 0.f;
-
-	UPROPERTY(EditInstanceOnly, BlueprintReadWrite, Setter, Getter, Category="Layout", meta=(InvalidEnumValues="Custom"))
-	ECEClonerAxis TwistAxis = ECEClonerAxis::Y;
-
-private:
 #if WITH_EDITOR
 	/** Used for PECP */
 	static const TCEPropertyChangeDispatcher<UCEClonerGridLayout> PropertyChangeDispatcher;
diff --git a/Engine/Plugins/Experimental/ClonerEffector/Source/ClonerEffector/Public/Cloner/Layouts/CEClonerLayoutBase.h b/Engine/Plugins/Experimental/ClonerEffector/Source/ClonerEffector/Public/Cloner/Layouts/CEClonerLayoutBase.h
--- a/Engine/Plugins/Experimental/ClonerEffector/Source/ClonerEffector/Public/Cloner/Layouts/CEClonerLayoutBase.h
+++ b/Engine/Plugins/Experimental/ClonerEffector/Source/ClonerEffector/Public/Cloner/Layouts/CEClonerLayoutBase.h
@@ -3,6 +3,7 @@
 #pragma once
 
 #include "CEClonerEffectorShared.h"
+#include "Templates/SubclassOf.h"
 #include "UObject/Object.h"
 #include "CEClonerLayoutBase.generated.h"
 
@@ -103,7 +104,13 @@ public:
 	bool CopyTo(UCEClonerLayoutBase* InOtherLayout) const;
 
 	/** Gets the cloner extensions supported by this layout */
-	virtual TSet<FName> GetSupportedExtensions() const;
+	TSet<TSubclassOf<UCEClonerExtensionBase>> GetSupportedExtensions() const;
+
+	/** Filter supported extension for this layout */
+	virtual bool IsExtensionSupported(const UCEClonerExtensionBase* InExtension) const
+	{
+		return true;
+	}
 
 	/** Request refresh layout next tick */
 	void MarkLayoutDirty(bool bInUpdateCloner = true);
diff --git a/Engine/Plugins/Experimental/ClonerEffector/Source/ClonerEffector/Public/Subsystems/CEClonerSubsystem.h b/Engine/Plugins/Experimental/ClonerEffector/Source/ClonerEffector/Public/Subsystems/CEClonerSubsystem.h
--- a/Engine/Plugins/Experimental/ClonerEffector/Source/ClonerEffector/Public/Subsystems/CEClonerSubsystem.h
+++ b/Engine/Plugins/Experimental/ClonerEffector/Source/ClonerEffector/Public/Subsystems/CEClonerSubsystem.h
@@ -58,7 +58,10 @@ public:
 	CLONEREFFECTOR_API bool IsExtensionClassRegistered(UClass* InClass) const;
 
 	/** Get available cloner extension names to use */
-	TArray<FName> GetExtensionNames() const;
+	TSet<FName> GetExtensionNames() const;
+
+	/** Get available cloner extension classes to use */
+	TSet<TSubclassOf<UCEClonerExtensionBase>> GetExtensionClasses() const;
 
 	/** Based on a extension class, find extension name */
 	FName FindExtensionName(TSubclassOf<UCEClonerExtensionBase> InClass) const;
