MotionDesign : ClonerEffector
- Cloner layout name and effector mode/type name once reset to default will select first available name instead of none
- Disable reset to default for effector channel
- Adding velocity output to cloner material function
- Adding base effector simulation for custom niagara simulation that want to use effectors without the cloner

#virtualized

diff --git a/Engine/Plugins/Experimental/ClonerEffector/Source/ClonerEffector/Private/Cloner/CEClonerComponent.cpp b/Engine/Plugins/Experimental/ClonerEffector/Source/ClonerEffector/Private/Cloner/CEClonerComponent.cpp
--- a/Engine/Plugins/Experimental/ClonerEffector/Source/ClonerEffector/Private/Cloner/CEClonerComponent.cpp
+++ b/Engine/Plugins/Experimental/ClonerEffector/Source/ClonerEffector/Private/Cloner/CEClonerComponent.cpp
@@ -73,8 +73,8 @@ UCEClonerComponent::UCEClonerComponent()
 #endif
 
 		// Apply default layout
-		const TArray<FString> LayoutNames = GetClonerLayoutNames();
-		LayoutName = !LayoutNames.IsEmpty() ? FName(LayoutNames[0]) : NAME_None;
+		const TArray<FName> LayoutNames = GetClonerLayoutNames();
+		LayoutName = !LayoutNames.IsEmpty() ? LayoutNames[0] : NAME_None;
 	}
 }
 
@@ -1140,7 +1140,7 @@ void UCEClonerComponent::SetLayoutName(FName InLayoutName)
 		return;
 	}
 
-	const TArray<FString> LayoutNames = GetClonerLayoutNames();
+	const TArray<FName> LayoutNames = GetClonerLayoutNames();
 	if (!LayoutNames.Contains(InLayoutName))
 	{
 		return;
@@ -1455,6 +1455,14 @@ void UCEClonerComponent::OnLayoutNameChanged()
 		return;
 	}
 
+	const TArray<FName> LayoutNames = GetClonerLayoutNames();
+
+	// Set default if value does not exists
+	if (!LayoutNames.Contains(LayoutName) && !LayoutNames.IsEmpty())
+	{
+		LayoutName = LayoutNames[0];
+	}
+
 	UCEClonerLayoutBase* NewActiveLayout = FindOrAddLayout(LayoutName);
 
 	// Apply layout
@@ -1577,22 +1585,16 @@ UCEClonerExtensionBase* UCEClonerComponent::FindOrAddExtension(FName InExtension
 	return NewActiveExtension;
 }
 
-TArray<FString> UCEClonerComponent::GetClonerLayoutNames() const
+TArray<FName> UCEClonerComponent::GetClonerLayoutNames() const
 {
-	TArray<FString> LayoutNamesStrings;
+	TArray<FName> LayoutNames;
 
 	if (const UCEClonerSubsystem* Subsystem = UCEClonerSubsystem::Get())
 	{
-		const TArray<FName>& LayoutNames = Subsystem->GetLayoutNames();
-		LayoutNamesStrings.Reserve(LayoutNames.Num());
-
-		Algo::Transform(LayoutNames, LayoutNamesStrings, [](const FName& InName)
-		{
-			return InName.ToString();
-		});
+		LayoutNames = Subsystem->GetLayoutNames();
 	}
 
-	return LayoutNamesStrings;
+	return LayoutNames;
 }
 
 void UCEClonerComponent::RefreshClonerMeshes()
diff --git a/Engine/Plugins/Experimental/ClonerEffector/Source/ClonerEffector/Private/Effector/CEEffectorComponent.cpp b/Engine/Plugins/Experimental/ClonerEffector/Source/ClonerEffector/Private/Effector/CEEffectorComponent.cpp
--- a/Engine/Plugins/Experimental/ClonerEffector/Source/ClonerEffector/Private/Effector/CEEffectorComponent.cpp
+++ b/Engine/Plugins/Experimental/ClonerEffector/Source/ClonerEffector/Private/Effector/CEEffectorComponent.cpp
@@ -493,6 +493,14 @@ void UCEEffectorComponent::OnEffectorOptionsChanged()
 
 void UCEEffectorComponent::OnTypeNameChanged()
 {
+	const TArray<FName> TypeNames = GetEffectorTypeNames();
+
+	// Set default if value does not exists
+	if (!TypeNames.Contains(TypeName) && !TypeNames.IsEmpty())
+	{
+		TypeName = TypeNames[0];
+	}
+
 	if (UCEEffectorTypeBase* Type = Cast<UCEEffectorTypeBase>(FindOrAddExtension(TypeName)))
 	{
 		if (ActiveType != Type)
@@ -519,6 +527,14 @@ void UCEEffectorComponent::OnTypeNameChanged()
 
 void UCEEffectorComponent::OnModeNameChanged()
 {
+	const TArray<FName> ModeNames = GetEffectorModeNames();
+
+	// Set default if value does not exists
+	if (!ModeNames.Contains(ModeName) && !ModeNames.IsEmpty())
+	{
+		ModeName = ModeNames[0];
+	}
+
 	if (UCEEffectorModeBase* Mode = Cast<UCEEffectorModeBase>(FindOrAddExtension(ModeName)))
 	{
 		if (ActiveMode != Mode)
diff --git a/Engine/Plugins/Experimental/ClonerEffector/Source/ClonerEffector/Public/CEClonerEffectorShared.h b/Engine/Plugins/Experimental/ClonerEffector/Source/ClonerEffector/Public/CEClonerEffectorShared.h
--- a/Engine/Plugins/Experimental/ClonerEffector/Source/ClonerEffector/Public/CEClonerEffectorShared.h
+++ b/Engine/Plugins/Experimental/ClonerEffector/Source/ClonerEffector/Public/CEClonerEffectorShared.h
@@ -440,7 +440,7 @@ struct FCEClonerEffectorChannelData
 
 protected:
 	/** Cached effector identifier to detect a change and update cloners DI */
-	UPROPERTY(EditInstanceOnly, Category="Effector")
+	UPROPERTY(VisibleInstanceOnly, Category="Effector", meta=(NoResetToDefault))
 	int32 Identifier = INDEX_NONE;
 
 	void Write(UNiagaraDataChannelWriter* InWriter) const;
diff --git a/Engine/Plugins/Experimental/ClonerEffector/Source/ClonerEffector/Public/Cloner/CEClonerComponent.h b/Engine/Plugins/Experimental/ClonerEffector/Source/ClonerEffector/Public/Cloner/CEClonerComponent.h
--- a/Engine/Plugins/Experimental/ClonerEffector/Source/ClonerEffector/Public/Cloner/CEClonerComponent.h
+++ b/Engine/Plugins/Experimental/ClonerEffector/Source/ClonerEffector/Public/Cloner/CEClonerComponent.h
@@ -395,7 +395,7 @@ private:
 
 	/** Gets all layout names available */
 	UFUNCTION()
-	TArray<FString> GetClonerLayoutNames() const;
+	TArray<FName> GetClonerLayoutNames() const;
 
 	/** Attachment tree view */
 	UPROPERTY(Transient, NonTransactional)
diff --git a/Engine/Plugins/Experimental/ClonerEffector/Source/ClonerEffector/Public/Effector/CEEffectorComponent.h b/Engine/Plugins/Experimental/ClonerEffector/Source/ClonerEffector/Public/Effector/CEEffectorComponent.h
--- a/Engine/Plugins/Experimental/ClonerEffector/Source/ClonerEffector/Public/Effector/CEEffectorComponent.h
+++ b/Engine/Plugins/Experimental/ClonerEffector/Source/ClonerEffector/Public/Effector/CEEffectorComponent.h
@@ -258,7 +258,7 @@ private:
 	TArray<FName> GetEffectorModeNames() const;
 
 	/** Transient effector channel data */
-	UPROPERTY(VisibleInstanceOnly, Transient, DuplicateTransient, TextExportTransient, NonTransactional, AdvancedDisplay, Category="Effector")
+	UPROPERTY(VisibleInstanceOnly, Transient, DuplicateTransient, TextExportTransient, NonTransactional, AdvancedDisplay, Category="Effector", meta=(NoResetToDefault))
 	FCEClonerEffectorChannelData ChannelData;
 
 	/** Cloners linked to this effector, used to refresh or relink on duplicate */
