MotionDesign : ClonerEffector
- Adding step mode for effector to transform clones based on their index
- Fixed wrong sphere weight computation
- Fixed low resolution visualizer meshes

#virtualized

diff --git a/Engine/Plugins/Experimental/ClonerEffector/Source/ClonerEffector/Private/Effector/Modes/CEEffectorStepMode.cpp b/Engine/Plugins/Experimental/ClonerEffector/Source/ClonerEffector/Private/Effector/Modes/CEEffectorStepMode.cpp
new file mode 100644
--- /dev/null
+++ b/Engine/Plugins/Experimental/ClonerEffector/Source/ClonerEffector/Private/Effector/Modes/CEEffectorStepMode.cpp
@@ -0,0 +1,66 @@
+// Copyright Epic Games, Inc. All Rights Reserved.
+
+#include "Effector/Modes/CEEffectorStepMode.h"
+
+#include "Effector/CEEffectorComponent.h"
+
+void UCEEffectorStepMode::SetStepPosition(const FVector& InPosition)
+{
+	if (StepPosition.Equals(InPosition))
+	{
+		return;
+	}
+
+	StepPosition = InPosition;
+	UpdateExtensionParameters();
+}
+
+void UCEEffectorStepMode::SetStepRotation(const FRotator& InRotation)
+{
+	if (StepRotation.Equals(InRotation))
+	{
+		return;
+	}
+
+	StepRotation = InRotation;
+	UpdateExtensionParameters();
+}
+
+void UCEEffectorStepMode::SetStepScale(const FVector& InScale)
+{
+	const FVector NewScale = InScale.ComponentMax(FVector(UE_KINDA_SMALL_NUMBER));
+
+	if (NewScale.Equals(StepScale))
+	{
+		return;
+	}
+
+	StepScale = NewScale;
+	UpdateExtensionParameters();
+}
+
+#if WITH_EDITOR
+const TCEPropertyChangeDispatcher<UCEEffectorStepMode> UCEEffectorStepMode::PropertyChangeDispatcher =
+{
+	{ GET_MEMBER_NAME_CHECKED(UCEEffectorStepMode, StepPosition), &UCEEffectorStepMode::OnExtensionPropertyChanged },
+	{ GET_MEMBER_NAME_CHECKED(UCEEffectorStepMode, StepRotation), &UCEEffectorStepMode::OnExtensionPropertyChanged },
+	{ GET_MEMBER_NAME_CHECKED(UCEEffectorStepMode, StepScale), &UCEEffectorStepMode::OnExtensionPropertyChanged },
+};
+
+void UCEEffectorStepMode::PostEditChangeProperty(FPropertyChangedEvent& InPropertyChangedEvent)
+{
+	Super::PostEditChangeProperty(InPropertyChangedEvent);
+
+	PropertyChangeDispatcher.OnPropertyChanged(this, InPropertyChangedEvent);
+}
+#endif
+
+void UCEEffectorStepMode::OnExtensionParametersChanged(UCEEffectorComponent* InComponent)
+{
+	Super::OnExtensionParametersChanged(InComponent);
+
+	FCEClonerEffectorChannelData& ChannelData = InComponent->GetChannelData();
+	ChannelData.LocationDelta = StepPosition;
+	ChannelData.RotationDelta = FVector(StepRotation.Yaw, StepRotation.Pitch, StepRotation.Roll);
+	ChannelData.ScaleDelta = StepScale.ComponentMax(FVector(UE_KINDA_SMALL_NUMBER));
+}
diff --git a/Engine/Plugins/Experimental/ClonerEffector/Source/ClonerEffector/Private/Effector/Types/CEEffectorSphereType.cpp b/Engine/Plugins/Experimental/ClonerEffector/Source/ClonerEffector/Private/Effector/Types/CEEffectorSphereType.cpp
--- a/Engine/Plugins/Experimental/ClonerEffector/Source/ClonerEffector/Private/Effector/Types/CEEffectorSphereType.cpp
+++ b/Engine/Plugins/Experimental/ClonerEffector/Source/ClonerEffector/Private/Effector/Types/CEEffectorSphereType.cpp
@@ -49,21 +49,23 @@ void UCEEffectorSphereType::OnExtensionVisualizerDirty(int32 InDirtyFlags)
 {
 	Super::OnExtensionVisualizerDirty(InDirtyFlags);
 
+	constexpr FGeometryScriptPrimitiveOptions PrimitiveOptions;
+	constexpr int32 Phi = 20;
+	constexpr int32 Theta = 32;
+
 	if ((InDirtyFlags & InnerVisualizerFlag) != 0)
 	{
-		UpdateVisualizer(InnerVisualizerFlag, [this](UDynamicMesh* InMesh)
+		UpdateVisualizer(InnerVisualizerFlag, [this, &PrimitiveOptions, Phi, Theta](UDynamicMesh* InMesh)
 		{
-			constexpr FGeometryScriptPrimitiveOptions PrimitiveOptions;
-			UGeometryScriptLibrary_MeshPrimitiveFunctions::AppendSphereLatLong(InMesh, PrimitiveOptions, FTransform::Identity, InnerRadius);
+			UGeometryScriptLibrary_MeshPrimitiveFunctions::AppendSphereLatLong(InMesh, PrimitiveOptions, FTransform::Identity, InnerRadius, Phi, Theta);
 		});
 	}
 
 	if ((InDirtyFlags & OuterVisualizerFlag) != 0)
 	{
-		UpdateVisualizer(OuterVisualizerFlag, [this](UDynamicMesh* InMesh)
+		UpdateVisualizer(OuterVisualizerFlag, [this, &PrimitiveOptions, Phi, Theta](UDynamicMesh* InMesh)
 		{
-			constexpr FGeometryScriptPrimitiveOptions PrimitiveOptions;
-			UGeometryScriptLibrary_MeshPrimitiveFunctions::AppendSphereLatLong(InMesh, PrimitiveOptions, FTransform::Identity, OuterRadius);
+			UGeometryScriptLibrary_MeshPrimitiveFunctions::AppendSphereLatLong(InMesh, PrimitiveOptions, FTransform::Identity, OuterRadius, Phi, Theta);
 		});
 	}
 }
diff --git a/Engine/Plugins/Experimental/ClonerEffector/Source/ClonerEffector/Private/Effector/Types/CEEffectorTorusType.cpp b/Engine/Plugins/Experimental/ClonerEffector/Source/ClonerEffector/Private/Effector/Types/CEEffectorTorusType.cpp
--- a/Engine/Plugins/Experimental/ClonerEffector/Source/ClonerEffector/Private/Effector/Types/CEEffectorTorusType.cpp
+++ b/Engine/Plugins/Experimental/ClonerEffector/Source/ClonerEffector/Private/Effector/Types/CEEffectorTorusType.cpp
@@ -64,23 +64,24 @@ void UCEEffectorTorusType::OnExtensionVisualizerDirty(int32 InDirtyFlags)
 {
 	Super::OnExtensionVisualizerDirty(InDirtyFlags);
 
+	constexpr FGeometryScriptPrimitiveOptions PrimitiveOptions;
+	constexpr FGeometryScriptRevolveOptions RevolveOptions;
+	constexpr int32 MajorSteps = 32;
+	constexpr int32 MinorSteps = 16;
+
 	if ((InDirtyFlags & InnerVisualizerFlag) != 0)
 	{
-		UpdateVisualizer(InnerVisualizerFlag, [this](UDynamicMesh* InMesh)
+		UpdateVisualizer(InnerVisualizerFlag, [this, &PrimitiveOptions, &RevolveOptions, MajorSteps, MinorSteps](UDynamicMesh* InMesh)
 		{
-			constexpr FGeometryScriptPrimitiveOptions PrimitiveOptions;
-			constexpr FGeometryScriptRevolveOptions RevolveOptions;
-			UGeometryScriptLibrary_MeshPrimitiveFunctions::AppendTorus(InMesh, PrimitiveOptions, FTransform(FVector(0, 0, -TorusInnerRadius)), RevolveOptions, TorusRadius, TorusInnerRadius);
+			UGeometryScriptLibrary_MeshPrimitiveFunctions::AppendTorus(InMesh, PrimitiveOptions, FTransform(FVector(0, 0, -TorusInnerRadius)), RevolveOptions, TorusRadius, TorusInnerRadius, MajorSteps, MinorSteps);
 		});
 	}
 
 	if ((InDirtyFlags & OuterVisualizerFlag) != 0)
 	{
-		UpdateVisualizer(OuterVisualizerFlag, [this](UDynamicMesh* InMesh)
+		UpdateVisualizer(OuterVisualizerFlag, [this, &PrimitiveOptions, &RevolveOptions, MajorSteps, MinorSteps](UDynamicMesh* InMesh)
 		{
-			constexpr FGeometryScriptPrimitiveOptions PrimitiveOptions;
-			constexpr FGeometryScriptRevolveOptions RevolveOptions;
-			UGeometryScriptLibrary_MeshPrimitiveFunctions::AppendTorus(InMesh, PrimitiveOptions, FTransform(FVector(0, 0, -TorusOuterRadius)), RevolveOptions, TorusRadius, TorusOuterRadius);
+			UGeometryScriptLibrary_MeshPrimitiveFunctions::AppendTorus(InMesh, PrimitiveOptions, FTransform(FVector(0, 0, -TorusOuterRadius)), RevolveOptions, TorusRadius, TorusOuterRadius, MajorSteps, MinorSteps);
 		});
 	}
 }
diff --git a/Engine/Plugins/Experimental/ClonerEffector/Source/ClonerEffector/Public/CEClonerEffectorShared.h b/Engine/Plugins/Experimental/ClonerEffector/Source/ClonerEffector/Public/CEClonerEffectorShared.h
--- a/Engine/Plugins/Experimental/ClonerEffector/Source/ClonerEffector/Public/CEClonerEffectorShared.h
+++ b/Engine/Plugins/Experimental/ClonerEffector/Source/ClonerEffector/Public/CEClonerEffectorShared.h
@@ -157,10 +157,12 @@ enum class ECEClonerEffectorMode : uint8
 	Default,
 	/** Rotates clones towards a target actor */
 	Target,
-	/** Randomly applies curl noise across the field zone */
+	/** Randomly applies noise across the field zone */
 	NoiseField,
-	/** Pushes clones apart based on a strength */
-	Push
+	/** Pushes clones apart based on a strength and direction */
+	Push,
+	/** Accumulate transform on clones based on their index */
+	Step
 };
 
 /** Enumerates the effector push direction available */
diff --git a/Engine/Plugins/Experimental/ClonerEffector/Source/ClonerEffector/Public/Effector/Modes/CEEffectorStepMode.h b/Engine/Plugins/Experimental/ClonerEffector/Source/ClonerEffector/Public/Effector/Modes/CEEffectorStepMode.h
new file mode 100644
--- /dev/null
+++ b/Engine/Plugins/Experimental/ClonerEffector/Source/ClonerEffector/Public/Effector/Modes/CEEffectorStepMode.h
@@ -0,0 +1,77 @@
+// Copyright Epic Games, Inc. All Rights Reserved.
+
+#pragma once
+
+#include "CEClonerEffectorShared.h"
+#include "CEEffectorModeBase.h"
+#include "CEPropertyChangeDispatcher.h"
+#include "CEEffectorStepMode.generated.h"
+
+class UCEEffectorComponent;
+
+UCLASS(MinimalAPI, BlueprintType, Within=CEEffectorComponent)
+class UCEEffectorStepMode : public UCEEffectorModeBase
+{
+	GENERATED_BODY()
+
+public:
+	UCEEffectorStepMode()
+		: UCEEffectorModeBase(TEXT("Step"), static_cast<int32>(ECEClonerEffectorMode::Step))
+	{}
+
+	UFUNCTION(BlueprintCallable, Category="Effector")
+	CLONEREFFECTOR_API void SetStepPosition(const FVector& InPosition);
+
+	UFUNCTION(BlueprintPure, Category="Effector")
+	FVector GetStepPosition() const
+	{
+		return StepPosition;
+	}
+
+	UFUNCTION(BlueprintCallable, Category="Effector")
+	CLONEREFFECTOR_API void SetStepRotation(const FRotator& InRotation);
+
+	UFUNCTION(BlueprintPure, Category="Effector")
+	FRotator GetStepRotation() const
+	{
+		return StepRotation;
+	}
+
+	UFUNCTION(BlueprintCallable, Category="Effector")
+	CLONEREFFECTOR_API void SetStepScale(const FVector& InScale);
+
+	UFUNCTION(BlueprintPure, Category="Effector")
+	FVector GetStepScale() const
+	{
+		return StepScale;
+	}
+
+protected:
+	//~ Begin UObject
+#if WITH_EDITOR
+	virtual void PostEditChangeProperty(FPropertyChangedEvent& InPropertyChangedEvent) override;
+#endif
+	//~ End UObject
+
+	//~ Begin UCEEffectorNoiseMode
+	virtual void OnExtensionParametersChanged(UCEEffectorComponent* InComponent) override;
+	//~ End UCEEffectorNoiseMode
+
+	/** Interpolates from 0 to this position offset based on the particle index and particle count */
+	UPROPERTY(EditInstanceOnly, BlueprintReadWrite, Setter, Getter, DisplayName="Position", Category="Mode")
+	FVector StepPosition = FVector::ZeroVector;
+
+	/** Interpolates from 0 to this rotation based on the particle index and particle count */
+	UPROPERTY(EditInstanceOnly, BlueprintReadWrite, Setter, Getter, DisplayName="Rotation", Category="Mode")
+	FRotator StepRotation = FRotator::ZeroRotator;
+
+	/** Interpolates from 1 to this scale based on the particle index and particle count */
+	UPROPERTY(EditInstanceOnly, BlueprintReadWrite, Setter, Getter, DisplayName="Scale", Category="Mode", meta=(ClampMin="0", MotionDesignVectorWidget, AllowPreserveRatio="XYZ"))
+	FVector StepScale = FVector::OneVector;
+
+private:
+#if WITH_EDITOR
+	/** Used for PECP */
+	static const TCEPropertyChangeDispatcher<UCEEffectorStepMode> PropertyChangeDispatcher;
+#endif
+};
\ No newline at end of file
