MotionDesign : ClonerEffector
- Fixed text issues with cloner where characters would not appear during rendering and invalid material warning would show up
- Adding MRQ preset for cloner to warmup and reduce motion ghosting due to TSR
- Replacing invalid materials by default cloner material when niagara flag is missing
- Removing effector data interface struct and moving logic to effector extension
- Improving material detection change by skipping components with no geometry
- Listening to transform delegates to catch character spacing changes and rebuild cloner mesh

#virtualized

diff --git a/Engine/Plugins/Experimental/ClonerEffector/Source/ClonerEffector/Private/CEClonerEffectorShared.cpp b/Engine/Plugins/Experimental/ClonerEffector/Source/ClonerEffector/Private/CEClonerEffectorShared.cpp
--- a/Engine/Plugins/Experimental/ClonerEffector/Source/ClonerEffector/Private/CEClonerEffectorShared.cpp
+++ b/Engine/Plugins/Experimental/ClonerEffector/Source/ClonerEffector/Private/CEClonerEffectorShared.cpp
@@ -9,7 +9,6 @@
 #include "Engine/Texture2D.h"
 #include "Engine/World.h"
 #include "NiagaraDataChannelAccessor.h"
-#include "NiagaraSystem.h"
 #include "PrimitiveSceneProxy.h"
 #include "SceneManagement.h"
 #include "UObject/UObjectGlobals.h"
@@ -105,105 +104,6 @@ void FCEClonerEffectorChannelData::Read(const UNiagaraDataChannelReader* InReade
 	VectorNoiseForceAmount = InReader->ReadFloat(VectorNoiseForceAmountName, Identifier, bIsValid);
 }
 
-FCEClonerEffectorDataInterfaces::FCEClonerEffectorDataInterfaces(const UNiagaraSystem* InSystem)
-{
-	check(InSystem);
-	const FNiagaraUserRedirectionParameterStore& UserParameterStore = InSystem->GetExposedParameters();
-
-	static const FNiagaraVariable EffectorIndexDIVar(FNiagaraTypeDefinition(UNiagaraDataInterfaceArrayInt32::StaticClass()), FCEClonerEffectorDataInterfaces::IndexName);
-	UNiagaraDataInterfaceArrayInt32* IndexArrayDI = CastChecked<UNiagaraDataInterfaceArrayInt32>(UserParameterStore.GetDataInterface(EffectorIndexDIVar));
-	DataInterfaces.Add(EffectorIndexDIVar.GetName(), IndexArrayDI);
-}
-
-void FCEClonerEffectorDataInterfaces::Clear() const
-{
-	if (UNiagaraDataInterfaceArrayInt32* IndexArray = GetIndexArray())
-	{
-		IndexArray->GetArrayReference().Empty();
-	}
-}
-
-void FCEClonerEffectorDataInterfaces::CopyTo(FCEClonerEffectorDataInterfaces& InOther) const
-{
-	for (const TPair<FName, TObjectPtr<UNiagaraDataInterface>>& DataInterfacePair : DataInterfaces)
-	{
-		const TObjectPtr<UNiagaraDataInterface> DataInterface = DataInterfacePair.Value;
-		const TObjectPtr<UNiagaraDataInterface>* OtherDataInterface = InOther.DataInterfaces.Find(DataInterfacePair.Key);
-
-		if (!OtherDataInterface || !DataInterface)
-		{
-			continue;
-		}
-
-		DataInterface->CopyTo(*OtherDataInterface);
-	}
-}
-
-void FCEClonerEffectorDataInterfaces::Resize(int32 InSize) const
-{
-	if (UNiagaraDataInterfaceArrayInt32* IndexArray = GetIndexArray())
-	{
-		IndexArray->GetArrayReference().SetNum(InSize);
-	}
-}
-
-void FCEClonerEffectorDataInterfaces::Remove(int32 InIndex) const
-{
-	UNiagaraDataInterfaceArrayInt32* IndexArray = GetIndexArray();
-	if (IndexArray && IndexArray->GetArrayReference().IsValidIndex(InIndex))
-	{
-		IndexArray->GetArrayReference().RemoveAt(InIndex);
-	}
-}
-
-bool FCEClonerEffectorDataInterfaces::IsValid() const
-{
-	for (const TPair<FName, TObjectPtr<UNiagaraDataInterface>>& DataInterfacePair : DataInterfaces)
-	{
-		if (!DataInterfacePair.Value.Get())
-		{
-			return false;
-		}
-	}
-
-	return !DataInterfaces.IsEmpty();
-}
-
-int32 FCEClonerEffectorDataInterfaces::Num() const
-{
-	UNiagaraDataInterfaceArrayInt32* IndexDI = GetIndexArray();
-
-	const int32 IndexSize = IndexDI->GetArrayReference().Num();
-
-	return IndexSize;
-}
-
-void FCEClonerEffectorDataInterfaces::Commit() const
-{
-	for (const TPair<FName, TObjectPtr<UNiagaraDataInterface>>& DataInterfacePair : DataInterfaces)
-	{
-		if (UNiagaraDataInterface* NiagaraDataInterface = DataInterfacePair.Value.Get())
-		{
-			if (UNiagaraSystem* System = NiagaraDataInterface->GetTypedOuter<UNiagaraSystem>())
-			{
-				FNiagaraUserRedirectionParameterStore& UserParameterStore = System->GetExposedParameters();
-				const FNiagaraVariable DIVar(FNiagaraTypeDefinition(NiagaraDataInterface->GetClass()), DataInterfacePair.Key);
-				UserParameterStore.SetDataInterface(NiagaraDataInterface, DIVar);
-			}
-		}
-	}
-}
-
-UNiagaraDataInterfaceArrayInt32* FCEClonerEffectorDataInterfaces::GetIndexArray() const
-{
-	if (const TObjectPtr<UNiagaraDataInterface>* Array = DataInterfaces.Find(FCEClonerEffectorDataInterfaces::IndexName))
-	{
-		return Cast<UNiagaraDataInterfaceArrayInt32>(Array->Get());
-	}
-
-	return nullptr;
-}
-
 #if WITH_EDITOR
 FCEExtensionSection UE::ClonerEffector::EditorSection::GetExtensionSectionFromClass(UClass* InClass)
 {
diff --git a/Engine/Plugins/Experimental/ClonerEffector/Source/ClonerEffector/Private/Cloner/CEClonerComponent.cpp b/Engine/Plugins/Experimental/ClonerEffector/Source/ClonerEffector/Private/Cloner/CEClonerComponent.cpp
--- a/Engine/Plugins/Experimental/ClonerEffector/Source/ClonerEffector/Private/Cloner/CEClonerComponent.cpp
+++ b/Engine/Plugins/Experimental/ClonerEffector/Source/ClonerEffector/Private/Cloner/CEClonerComponent.cpp
@@ -7,7 +7,6 @@
 #include "Cloner/Extensions/CEClonerExtensionBase.h"
 #include "Cloner/Layouts/CEClonerLayoutBase.h"
 #include "Components/BillboardComponent.h"
-#include "Components/BrushComponent.h"
 #include "Components/DynamicMeshComponent.h"
 #include "Components/SkeletalMeshComponent.h"
 #include "Engine/StaticMeshActor.h"
@@ -16,7 +15,6 @@
 #include "Materials/MaterialInstanceDynamic.h"
 #include "NiagaraMeshRendererProperties.h"
 #include "NiagaraSystem.h"
-#include "ProceduralMeshComponent.h"
 #include "Settings/CEClonerEffectorSettings.h"
 #include "Subsystems/CEClonerSubsystem.h"
 #include "UDynamicMesh.h"
@@ -65,6 +63,7 @@ UCEClonerComponent::UCEClonerComponent()
 	if (!IsTemplate())
 	{
 		UCEClonerSubsystem::OnClonerSetEnabled().AddUObject(this, &UCEClonerComponent::OnClonerSetEnabled);
+		USceneComponent::MarkRenderStateDirtyEvent.AddUObject(this, &UCEClonerComponent::OnRenderStateDirty);
 
 		// Bind to delegate to detect material changes
 #if WITH_EDITOR
@@ -249,27 +248,30 @@ void UCEClonerComponent::UpdateAttachmentTree()
 	// Did we need to update meshes
 	TArray<TWeakObjectPtr<AActor>> ClonedActors;
 	ClonerTree.ItemAttachmentMap.GenerateKeyArray(ClonedActors);
-	for (const TWeakObjectPtr<AActor>& ClonedActor : ClonedActors)
+	for (const TWeakObjectPtr<AActor>& ClonedActorWeak : ClonedActors)
 	{
-		FCEClonerAttachmentItem* ClonedItem = ClonerTree.ItemAttachmentMap.Find(ClonedActor);
+		FCEClonerAttachmentItem* ClonedItem = ClonerTree.ItemAttachmentMap.Find(ClonedActorWeak);
 
 		if (!ClonedItem)
 		{
 			continue;
 		}
 
+		AActor* ClonedActor = ClonedActorWeak.Get();
+
 		if (ClonedItem->Status == ECEClonerAttachmentStatus::Invalid)
 		{
-			InvalidateBakedStaticMesh(ClonedActor.Get());
-			UnbindActorDelegates(ClonedActor.Get());
+			InvalidateBakedStaticMesh(ClonedActor);
+			UnbindActorDelegates(ClonedActor);
 			ClonerTree.ItemAttachmentMap.Remove(ClonedItem->ItemActor);
+			SetActorVisibility(ClonedActor, true);
 		}
 		else if (ClonedItem->Status == ECEClonerAttachmentStatus::Outdated)
 		{
 			if (ClonedItem->MeshStatus == ECEClonerAttachmentStatus::Outdated)
 			{
 				ClonerTree.DirtyItemAttachments.Add(ClonedItem);
-				InvalidateBakedStaticMesh(ClonedActor.Get());
+				InvalidateBakedStaticMesh(ClonedActor);
 			}
 			bClonerMeshesDirty = true;
 			ClonedItem->Status = ECEClonerAttachmentStatus::Updated;
@@ -353,6 +355,7 @@ void UCEClonerComponent::UpdateActorAttachment(AActor* InActor, AActor* InParent
 		AttachmentItem->Status = ECEClonerAttachmentStatus::Outdated;
 		InvalidateBakedStaticMesh(InActor);
 		BindActorDelegates(InActor);
+		SetActorVisibility(InActor, false);
 	}
 
 	if (AttachmentItem->ChildrenActors.Num() != ChildrenActors.Num())
@@ -375,14 +378,10 @@ void UCEClonerComponent::BindActorDelegates(AActor* InActor)
 		return;
 	}
 
-#if WITH_EDITOR
-	// Hide new actor when attached to cloner
-	InActor->SetIsTemporarilyHiddenInEditor(true);
-#endif
-	InActor->SetActorHiddenInGame(true);
-
-	InActor->OnDestroyed.RemoveAll(this);
-	InActor->OnDestroyed.AddDynamic(this, &UCEClonerComponent::OnActorDestroyed);
+	if (!InActor->OnDestroyed.IsAlreadyBound(this, &UCEClonerComponent::OnActorDestroyed))
+	{
+		InActor->OnDestroyed.AddDynamic(this, &UCEClonerComponent::OnActorDestroyed);
+	}
 
 #if WITH_EDITOR
 	// Detect static mesh change
@@ -408,6 +407,17 @@ void UCEClonerComponent::BindActorDelegates(AActor* InActor)
 			DynamicMeshComponent->OnMeshChanged.AddUObject(this, &UCEClonerComponent::OnMeshChanged, NullComponent, InActor);
 		}
 	}
+
+	// Detect components transform
+	TArray<USceneComponent*> SceneComponents;
+	InActor->GetComponents(SceneComponents, /** IncludeChildren */false);
+	for (USceneComponent* SceneComponent : SceneComponents)
+	{
+		if (!SceneComponent->TransformUpdated.IsBoundToObject(this))
+		{
+			SceneComponent->TransformUpdated.AddUObject(this, &UCEClonerComponent::OnComponentTransformed);
+		}
+	}
 }
 
 void UCEClonerComponent::UnbindActorDelegates(AActor* InActor) const
@@ -417,12 +427,6 @@ void UCEClonerComponent::UnbindActorDelegates(AActor* InActor) const
 		return;
 	}
 
-#if WITH_EDITOR
-	// Show new actor when detached to cloner
-	InActor->SetIsTemporarilyHiddenInEditor(false);
-#endif
-	InActor->SetActorHiddenInGame(false);
-
 	InActor->OnDestroyed.RemoveAll(this);
 
 #if WITH_EDITOR
@@ -440,6 +444,26 @@ void UCEClonerComponent::UnbindActorDelegates(AActor* InActor) const
 	{
 		DynamicMeshComponent->OnMeshChanged.RemoveAll(this);
 	}
+
+	TArray<USceneComponent*> SceneComponents;
+	InActor->GetComponents(SceneComponents, /** IncludeChildren */false);
+	for (USceneComponent* SceneComponent : SceneComponents)
+	{
+		SceneComponent->TransformUpdated.RemoveAll(this);
+	}
+}
+
+void UCEClonerComponent::SetActorVisibility(AActor* InActor, bool bInVisibility)
+{
+	if (!InActor)
+	{
+		return;
+	}
+
+#if WITH_EDITOR
+	InActor->SetIsTemporarilyHiddenInEditor(!bInVisibility);
+#endif
+	InActor->SetActorHiddenInGame(!bInVisibility);
 }
 
 void UCEClonerComponent::OnActorDestroyed(AActor* InDestroyedActor)
@@ -449,6 +473,7 @@ void UCEClonerComponent::OnActorDestroyed(AActor* InDestroyedActor)
 		InvalidateBakedStaticMesh(InDestroyedActor);
 		UnbindActorDelegates(InDestroyedActor);
 		ClonerTree.ItemAttachmentMap.Remove(InDestroyedActor);
+		SetActorVisibility(InDestroyedActor, true);
 		bClonerMeshesDirty = true;
 	}
 }
@@ -489,84 +514,30 @@ void UCEClonerComponent::OnMaterialChanged(UObject* InObject)
 		return;
 	}
 
+	TArray<UPrimitiveComponent*> PrimitiveComponents;
+	ActorChanged->GetComponents(PrimitiveComponents, /** IncludeChildrenActors */false);
+
 	int32 MatIdx = 0;
-	TArray<TWeakObjectPtr<UMaterialInterface>> NewMaterials;
 	bool bMaterialChanged = false;
+	TArray<TWeakObjectPtr<UMaterialInterface>> NewMaterials;
+	NewMaterials.Reserve(PrimitiveComponents.Num());
 
-	TArray<UDynamicMeshComponent*> DynamicMeshComponents;
-	ActorChanged->GetComponents(DynamicMeshComponents, false);
-	for (const UDynamicMeshComponent* DynamicMeshComponent : DynamicMeshComponents)
-	{
-		for (UMaterialInterface* Material : DynamicMeshComponent->GetMaterials())
-		{
-			if (!AttachmentItem->BakedMaterials.IsValidIndex(MatIdx) || AttachmentItem->BakedMaterials[MatIdx] != Material)
-			{
-				bMaterialChanged = true;
-			}
-			NewMaterials.Add(Material);
-			MatIdx++;
-		}
-	}
-
-	TArray<USkeletalMeshComponent*> SkeletalMeshComponents;
-	ActorChanged->GetComponents(SkeletalMeshComponents, false);
-	for (const USkeletalMeshComponent* SkeletalMeshComponent : SkeletalMeshComponents)
+	for (UPrimitiveComponent* PrimitiveComponent : PrimitiveComponents)
 	{
-		for (UMaterialInterface* Material : SkeletalMeshComponent->GetMaterials())
+		if (!PrimitiveComponent || !FCEClonerMeshBuilder::HasAnyGeometry(PrimitiveComponent))
 		{
-			if (!AttachmentItem->BakedMaterials.IsValidIndex(MatIdx) || AttachmentItem->BakedMaterials[MatIdx] != Material)
-			{
-				bMaterialChanged = true;
-			}
-			NewMaterials.Add(Material);
-			MatIdx++;
+			continue;
 		}
-	}
 
-	TArray<UBrushComponent*> BrushComponents;
-	ActorChanged->GetComponents(BrushComponents, false);
-	for (const UBrushComponent* BrushComponent : BrushComponents)
-	{
-		for (int32 Idx = 0; Idx < BrushComponent->GetNumMaterials(); Idx++)
+		for (int32 MatIndex = 0; MatIndex < PrimitiveComponent->GetNumMaterials(); MatIndex++)
 		{
-			if (!AttachmentItem->BakedMaterials.IsValidIndex(MatIdx) || AttachmentItem->BakedMaterials[MatIdx] != BrushComponent->GetMaterial(Idx))
-			{
-				bMaterialChanged = true;
-			}
-			NewMaterials.Add(BrushComponent->GetMaterial(Idx));
-			MatIdx++;
-		}
-	}
+			UMaterialInterface* Material = PrimitiveComponent->GetMaterial(MatIndex);
 
-	TArray<UProceduralMeshComponent*> ProceduralMeshComponents;
-	ActorChanged->GetComponents(ProceduralMeshComponents, false);
-	for (const UProceduralMeshComponent* ProceduralMeshComponent : ProceduralMeshComponents)
-	{
-		if (ProceduralMeshComponent->GetNumSections() == 0)
-		{
-			continue;
-		}
-		for (UMaterialInterface* Material : ProceduralMeshComponent->GetMaterials())
-		{
 			if (!AttachmentItem->BakedMaterials.IsValidIndex(MatIdx) || AttachmentItem->BakedMaterials[MatIdx] != Material)
 			{
 				bMaterialChanged = true;
 			}
-			NewMaterials.Add(Material);
-			MatIdx++;
-		}
-	}
 
-	TArray<UStaticMeshComponent*> StaticMeshComponents;
-	ActorChanged->GetComponents(StaticMeshComponents, false);
-	for (const UStaticMeshComponent* StaticMeshComponent : StaticMeshComponents)
-	{
-		for (UMaterialInterface* Material : StaticMeshComponent->GetMaterials())
-		{
-			if (!AttachmentItem->BakedMaterials.IsValidIndex(MatIdx) || AttachmentItem->BakedMaterials[MatIdx] != Material)
-			{
-				bMaterialChanged = true;
-			}
 			NewMaterials.Add(Material);
 			MatIdx++;
 		}
@@ -589,11 +560,6 @@ void UCEClonerComponent::OnMaterialChanged(UObject* InObject)
 	}
 }
 
-void UCEClonerComponent::OnTransformUpdated(USceneComponent* InUpdatedComponent, EUpdateTransformFlags InUpdateTransformFlags, ETeleportType InTeleport)
-{
-	InvalidateBakedStaticMesh(InUpdatedComponent->GetOwner());
-}
-
 void UCEClonerComponent::OnMeshChanged(UStaticMeshComponent*, AActor* InActor)
 {
 	if (!InActor)
@@ -853,9 +819,8 @@ void UCEClonerComponent::UpdateRootActorBakedStaticMesh(AActor* InRootActor)
 	TArray<FCEClonerAttachmentItem*> AttachmentItems;
 	GetActorAttachmentItems(InRootActor, AttachmentItems);
 
-#if WITH_EDITOR
 	int32 ReadOnlyMaterialCount = 0;
-#endif
+	UMaterialInterface* DefaultClonerMaterial = nullptr;
 
 	for (FCEClonerAttachmentItem* AttachmentItem : AttachmentItems)
 	{
@@ -882,21 +847,25 @@ void UCEClonerComponent::UpdateRootActorBakedStaticMesh(AActor* InRootActor)
 		{
 			UMaterialInterface* MaterialInterface = BakedMaterial.Get();
 
-#if WITH_EDITOR
 			if (MaterialInterface && !IsMaterialUsageFlagSet(MaterialInterface))
 			{
 				if (!IsMaterialDirtyable(MaterialInterface))
 				{
 					ReadOnlyMaterialCount++;
-					MaterialInterface = nullptr;
-					UE_LOG(LogCEClonerComponent, Warning, TEXT("%s : The following materials (%s) on actor (%s) does not have the required usage flag (bUsedWithNiagaraMeshParticles) to work with the cloner, this material cannot be dirtied due to its read-only location, skipping material and proceeding"), *ClonerActor->GetActorNameOrLabel(), MaterialInterface ? *MaterialInterface->GetMaterial()->GetPathName() : TEXT("Invalid Material"), *InRootActor->GetActorNameOrLabel());
+					UE_LOG(LogCEClonerComponent, Warning, TEXT("%s : The following materials (%s) on actor (%s) does not have the required usage flag (bUsedWithNiagaraMeshParticles) to work with the cloner, this material cannot be dirtied due to its read-only location, replacing material by default cloner material and proceeding"), *ClonerActor->GetActorNameOrLabel(), MaterialInterface ? *MaterialInterface->GetMaterial()->GetPathName() : TEXT("Invalid Material"), *InRootActor->GetActorNameOrLabel());
+
+					if (!DefaultClonerMaterial)
+					{
+						DefaultClonerMaterial = LoadObject<UMaterialInterface>(nullptr, UCEClonerEffectorSettings::DefaultMaterialPath);
+					}
+
+					MaterialInterface = DefaultClonerMaterial;
 				}
 				else
 				{
-					UE_LOG(LogCEClonerComponent, Log, TEXT("%s : The following materials (%s) on actor (%s) does not have the required usage flag (bUsedWithNiagaraMeshParticles) to work with the cloner, dirtying material and proceeding, please save the asset for working runtime result"), *ClonerActor->GetActorNameOrLabel(), MaterialInterface ? *MaterialInterface->GetMaterial()->GetPathName() : TEXT("Invalid Material"), *InRootActor->GetActorNameOrLabel());
+					UE_LOG(LogCEClonerComponent, Log, TEXT("%s : The following materials (%s) on actor (%s) does not have the required usage flag (bUsedWithNiagaraMeshParticles) to work with the cloner, setting the flag on material and proceeding, please resave the asset with the flag set"), *ClonerActor->GetActorNameOrLabel(), MaterialInterface ? *MaterialInterface->GetMaterial()->GetPathName() : TEXT("Invalid Material"), *InRootActor->GetActorNameOrLabel());
 				}
 			}
-#endif
 
 			MeshMaterials.Add(MaterialInterface);
 		}
@@ -911,12 +880,22 @@ void UCEClonerComponent::UpdateRootActorBakedStaticMesh(AActor* InRootActor)
 
 	ClonerTree.MergedBakedMeshes[RootIdx] = Mesh;
 
-#if WITH_EDITOR
 	if (ReadOnlyMaterialCount > 0)
 	{
-		ShowMaterialWarning(ReadOnlyMaterialCount);
-	}
+		UE_LOG(LogCEClonerComponent, Warning, TEXT("%s : %i invalid material(s) detected due to missing niagara usage flag (bUsedWithNiagaraMeshParticles) on actor (%s)"), *ClonerActor->GetActorNameOrLabel(), ReadOnlyMaterialCount, *InRootActor->GetActorNameOrLabel());
+
+#if WITH_EDITOR
+		// Fire warning notification when invalid materials are found and at least 5s has elapsed since last notification
+		constexpr double MinNotificationElapsedTime = 5.0;
+		const double CurrentTime = FApp::GetCurrentTime();
+
+		if (CurrentTime - LastNotificationTime > MinNotificationElapsedTime)
+		{
+			LastNotificationTime = CurrentTime;
+			ShowMaterialWarning(ReadOnlyMaterialCount);
+		}
 #endif
+	}
 }
 
 void UCEClonerComponent::GetActorAttachmentItems(AActor* InActor, TArray<FCEClonerAttachmentItem*>& OutAttachmentItems)
@@ -1371,11 +1350,8 @@ void UCEClonerComponent::RequestClonerUpdate(bool bInImmediate)
 	{
 		bNeedsRefresh = false;
 
-		if (UNiagaraSystem* ActiveSystem = GetAsset())
-		{
-			FNiagaraUserRedirectionParameterStore& UserParameterStore = ActiveSystem->GetExposedParameters();
-			UserParameterStore.PostGenericEditChange();
-		}
+		FNiagaraUserRedirectionParameterStore& UserParameterStore = GetOverrideParameters();
+		UserParameterStore.PostGenericEditChange();
 	}
 	else
 	{
@@ -1489,6 +1465,88 @@ void UCEClonerComponent::OnVisualizerSpriteVisibleChanged()
 }
 #endif
 
+void UCEClonerComponent::OnRenderStateDirty(UActorComponent& InActorComponent)
+{
+	AActor* Owner = InActorComponent.GetOwner();
+	const AActor* ClonerActor = GetOwner();
+
+	if (!Owner || Owner->GetLevel() != ClonerActor->GetLevel())
+	{
+		return;
+	}
+
+	// Does it contain geometry that we can convert
+	if (!FCEClonerMeshBuilder::IsComponentSupported(&InActorComponent))
+	{
+		return;
+	}
+
+	FCEClonerAttachmentItem* Item = ClonerTree.ItemAttachmentMap.Find(Owner);
+
+	if (!Item)
+	{
+		return;
+	}
+
+	UE_LOG(LogCEClonerComponent, Log, TEXT("%s : Render state changed for %s"), *ClonerActor->GetActorNameOrLabel(), *Owner->GetActorNameOrLabel());
+
+	// Rebind delegates as new components might be available
+	BindActorDelegates(Owner);
+
+	Item->MeshStatus = ECEClonerAttachmentStatus::Outdated;
+	InvalidateBakedStaticMesh(Owner);
+	ClonerTree.DirtyItemAttachments.Add(Item);
+}
+
+void UCEClonerComponent::OnComponentTransformed(USceneComponent* InComponent, EUpdateTransformFlags InFlags, ETeleportType InTeleport)
+{
+	if (!InComponent || !InComponent->GetOwner() || InFlags == EUpdateTransformFlags::PropagateFromParent)
+	{
+		return;
+	}
+
+	AActor* Owner = InComponent->GetOwner();
+	const AActor* RootActor = GetRootActor(Owner);
+
+	// Skip update if root component has moved, since we can simply offset the mesh
+	if (!RootActor || (RootActor == Owner && RootActor->GetRootComponent() == InComponent))
+	{
+		return;
+	}
+
+	bool bComponentSupported = FCEClonerMeshBuilder::IsComponentSupported(InComponent);
+
+	if (!bComponentSupported)
+	{
+		for (const TObjectPtr<USceneComponent>& ChildComponent : InComponent->GetAttachChildren())
+		{
+			if (FCEClonerMeshBuilder::IsComponentSupported(ChildComponent))
+			{
+				bComponentSupported = true;
+				break;
+			}
+		}
+	}
+
+	if (!bComponentSupported)
+	{
+		return;
+	}
+
+	FCEClonerAttachmentItem* Item = ClonerTree.ItemAttachmentMap.Find(Owner);
+
+	if (!Item)
+	{
+		return;
+	}
+
+	UE_LOG(LogCEClonerComponent, Log, TEXT("%s : Transform state changed for %s"), *GetOwner()->GetActorNameOrLabel(), *Owner->GetActorNameOrLabel());
+
+	Item->MeshStatus = ECEClonerAttachmentStatus::Outdated;
+	InvalidateBakedStaticMesh(Owner);
+	ClonerTree.DirtyItemAttachments.Add(Item);
+}
+
 UCEClonerLayoutBase* UCEClonerComponent::FindOrAddLayout(TSubclassOf<UCEClonerLayoutBase> InClass)
 {
 	const UCEClonerSubsystem* Subsystem = UCEClonerSubsystem::Get();
@@ -1673,12 +1731,6 @@ void UCEClonerComponent::ActivateLayout(UCEClonerLayoutBase* InLayout)
 		return;
 	}
 
-	// Copy data interfaces to new layout
-	if (ActiveLayout && ActiveLayout->IsLayoutLoaded())
-	{
-		ActiveLayout->CopyTo(InLayout);
-	}
-
 	// Deactivate previous layout
 	if (ActiveLayout && ActiveLayout->IsLayoutActive())
 	{
@@ -1758,7 +1810,7 @@ void UCEClonerComponent::UpdateClonerMeshes()
 
 	if (ActiveLayout->GetSystem() != ActiveSystem)
 	{
-		UE_LOG(LogCEClonerComponent, Warning, TEXT("%s : Invalid asset for cloner layout"), *ClonerActor->GetActorNameOrLabel());
+		UE_LOG(LogCEClonerComponent, Warning, TEXT("%s : Invalid system for cloner layout"), *ClonerActor->GetActorNameOrLabel());
 		return;
 	}
 
diff --git a/Engine/Plugins/Experimental/ClonerEffector/Source/ClonerEffector/Private/Cloner/CEClonerMeshBuilder.cpp b/Engine/Plugins/Experimental/ClonerEffector/Source/ClonerEffector/Private/Cloner/CEClonerMeshBuilder.cpp
--- a/Engine/Plugins/Experimental/ClonerEffector/Source/ClonerEffector/Private/Cloner/CEClonerMeshBuilder.cpp
+++ b/Engine/Plugins/Experimental/ClonerEffector/Source/ClonerEffector/Private/Cloner/CEClonerMeshBuilder.cpp
@@ -17,17 +17,121 @@
 #include "GeometryScript/MeshBasicEditFunctions.h"
 #include "GeometryScript/SceneUtilityFunctions.h"
 #include "Materials/MaterialInterface.h"
+#include "Model.h"
 #include "NiagaraComponent.h"
 #include "NiagaraEmitter.h"
 #include "NiagaraMeshRendererProperties.h"
 #include "NiagaraSimCacheFunctionLibrary.h"
 #include "NiagaraSystem.h"
 #include "ProceduralMeshComponent.h"
+#include "Rendering/SkeletalMeshRenderData.h"
 #include "StaticMeshOperations.h"
 #include "UDynamicMesh.h"
 
 const FCEClonerMeshBuilder::FCEClonerMeshBuilderParams FCEClonerMeshBuilder::DefaultParams;
 
+bool FCEClonerMeshBuilder::HasAnyGeometry(UActorComponent* InComponent)
+{
+	if (!IsComponentSupported(InComponent))
+	{
+		return false;
+	}
+
+	if (UDynamicMeshComponent* DynamicMeshComponent = Cast<UDynamicMeshComponent>(InComponent))
+	{
+		return DynamicMeshComponent->GetDynamicMesh()
+			&& DynamicMeshComponent->GetDynamicMesh()->GetTriangleCount() > 0;
+	}
+
+	if (const USkeletalMeshComponent* SkeletalMeshComponent = Cast<USkeletalMeshComponent>(InComponent))
+	{
+		if (SkeletalMeshComponent->GetSkeletalMeshAsset())
+		{
+			if (FSkeletalMeshRenderData* RenderData = SkeletalMeshComponent->GetSkeletalMeshRenderData())
+			{
+				return !RenderData->LODRenderData.IsEmpty() && RenderData->LODRenderData[0].GetNumVertices() > 0;
+			}
+		}
+
+		return false;
+	}
+
+	if (const UBrushComponent* BrushComponent = Cast<UBrushComponent>(InComponent))
+	{
+		return BrushComponent->Brush
+			&& !BrushComponent->Brush->Verts.IsEmpty();
+	}
+
+	if (UProceduralMeshComponent* ProceduralMeshComponent = Cast<UProceduralMeshComponent>(InComponent))
+	{
+		for (int32 SectionIndex = 0; SectionIndex < ProceduralMeshComponent->GetNumSections(); SectionIndex++)
+		{
+			if (const FProcMeshSection* Section = ProceduralMeshComponent->GetProcMeshSection(SectionIndex))
+			{
+				if (Section->bSectionVisible && !Section->ProcVertexBuffer.IsEmpty() && !Section->ProcIndexBuffer.IsEmpty())
+				{
+					return true;
+				}
+			}
+		}
+
+		return false;
+	}
+
+	if (const UInstancedStaticMeshComponent* ISMComponent = Cast<UInstancedStaticMeshComponent>(InComponent))
+	{
+		if (ISMComponent->GetStaticMesh() && ISMComponent->GetStaticMesh()->GetNumTriangles(/**LOD*/0) > 0)
+		{
+			return ISMComponent->GetNumInstances() > 0;
+		}
+
+		return false;
+	}
+
+	if (const USplineMeshComponent* SplineMeshComponent = Cast<USplineMeshComponent>(InComponent))
+	{
+		return SplineMeshComponent->GetStaticMesh() && SplineMeshComponent->GetStaticMesh()->GetNumTriangles(/**LOD*/0) > 0;
+	}
+
+	if (const UStaticMeshComponent* StaticMeshComponent = Cast<UStaticMeshComponent>(InComponent))
+	{
+		return StaticMeshComponent->GetStaticMesh() && StaticMeshComponent->GetStaticMesh()->GetNumTriangles(/**LOD*/0) > 0;
+	}
+
+	if (const UNiagaraComponent* NiagaraComponent = Cast<UNiagaraComponent>(InComponent))
+	{
+		if (UNiagaraSystem* System = NiagaraComponent->GetAsset())
+		{
+			if (System->GetActiveInstancesCount() > 0)
+			{
+				for (const FNiagaraEmitterHandle& EmitterHandle : System->GetEmitterHandles())
+				{
+					if (const FVersionedNiagaraEmitterData* EmitterData = EmitterHandle.GetEmitterData())
+					{
+						for (UNiagaraRendererProperties* EmitterRenderer : EmitterData->GetRenderers())
+						{
+							if (UNiagaraMeshRendererProperties* MeshRenderer = Cast<UNiagaraMeshRendererProperties>(EmitterRenderer))
+							{
+								for (const FNiagaraMeshRendererMeshProperties& MeshProperty : MeshRenderer->Meshes)
+								{
+									if (MeshProperty.Mesh && MeshProperty.Mesh->GetNumTriangles(/**LOD*/0) > 0)
+									{
+										return true;
+									}
+								}
+							}
+						}
+					}
+				}
+			}
+		}
+
+		return false;
+	}
+
+	return false;
+}
+
 FCEClonerMeshBuilder::FCEClonerMeshBuilder()
 {
 	OutputDynamicMesh = NewObject<UDynamicMesh>();
@@ -40,6 +144,16 @@ TArray<uint32> FCEClonerMeshBuilder::GetMeshIndexes() const
 	return MeshIndexes;
 }
 
+bool FCEClonerMeshBuilder::IsActorSupported(const AActor* InActor)
+{
+	return InActor && InActor->FindComponentByClass<UPrimitiveComponent>();
+}
+
+bool FCEClonerMeshBuilder::IsComponentSupported(const UActorComponent* InComponent)
+{
+	return InComponent && InComponent->IsA<UPrimitiveComponent>();
+}
+
 void FCEClonerMeshBuilder::Reset()
 {
 	ClearOutputMesh();
diff --git a/Engine/Plugins/Experimental/ClonerEffector/Source/ClonerEffector/Private/Cloner/Extensions/CEClonerCollisionExtension.cpp b/Engine/Plugins/Experimental/ClonerEffector/Source/ClonerEffector/Private/Cloner/Extensions/CEClonerCollisionExtension.cpp
--- a/Engine/Plugins/Experimental/ClonerEffector/Source/ClonerEffector/Private/Cloner/Extensions/CEClonerCollisionExtension.cpp
+++ b/Engine/Plugins/Experimental/ClonerEffector/Source/ClonerEffector/Private/Cloner/Extensions/CEClonerCollisionExtension.cpp
@@ -162,7 +162,7 @@ void UCEClonerCollisionExtension::OnExtensionParametersChanged(UCEClonerComponen
 			}
 		}
 
-		const FNiagaraUserRedirectionParameterStore& ExposedParameters = LayoutSystem->GetSystem()->GetExposedParameters();
+		const FNiagaraUserRedirectionParameterStore& ExposedParameters = InComponent->GetOverrideParameters();
 
 		static const FNiagaraVariable CollisionRadiiVar(FNiagaraTypeDefinition(UNiagaraDataInterfaceArrayFloat::StaticClass()), TEXT("CollisionRadii"));
 
diff --git a/Engine/Plugins/Experimental/ClonerEffector/Source/ClonerEffector/Private/Cloner/Extensions/CEClonerConstraintExtension.cpp b/Engine/Plugins/Experimental/ClonerEffector/Source/ClonerEffector/Private/Cloner/Extensions/CEClonerConstraintExtension.cpp
--- a/Engine/Plugins/Experimental/ClonerEffector/Source/ClonerEffector/Private/Cloner/Extensions/CEClonerConstraintExtension.cpp
+++ b/Engine/Plugins/Experimental/ClonerEffector/Source/ClonerEffector/Private/Cloner/Extensions/CEClonerConstraintExtension.cpp
@@ -168,7 +168,7 @@ void UCEClonerConstraintExtension::OnExtensionParametersChanged(UCEClonerCompone
 {
 	Super::OnExtensionParametersChanged(InComponent);
 
-	FNiagaraUserRedirectionParameterStore& ExposedParameters = InComponent->GetAsset()->GetExposedParameters();
+	FNiagaraUserRedirectionParameterStore& ExposedParameters = InComponent->GetOverrideParameters();
 	static const FNiagaraVariable ConstraintVar(FNiagaraTypeDefinition(StaticEnum<ECEClonerGridConstraint>()), TEXT("Constraint"));
 	ExposedParameters.SetParameterValue<int32>(static_cast<int32>(Constraint), ConstraintVar);
 
diff --git a/Engine/Plugins/Experimental/ClonerEffector/Source/ClonerEffector/Private/Cloner/Extensions/CEClonerEffectorExtension.cpp b/Engine/Plugins/Experimental/ClonerEffector/Source/ClonerEffector/Private/Cloner/Extensions/CEClonerEffectorExtension.cpp
--- a/Engine/Plugins/Experimental/ClonerEffector/Source/ClonerEffector/Private/Cloner/Extensions/CEClonerEffectorExtension.cpp
+++ b/Engine/Plugins/Experimental/ClonerEffector/Source/ClonerEffector/Private/Cloner/Extensions/CEClonerEffectorExtension.cpp
@@ -6,6 +6,7 @@
 #include "Cloner/Layouts/CEClonerLayoutBase.h"
 #include "Effector/CEEffectorActor.h"
 #include "Effector/CEEffectorComponent.h"
+#include "NiagaraDataInterfaceArrayInt.h"
 #include "Subsystems/CEEffectorSubsystem.h"
 
 #if WITH_EDITOR
@@ -67,24 +68,12 @@ int32 UCEClonerEffectorExtension::GetEffectorCount() const
 	return EffectorActorsWeak.Num();
 }
 
-const FCEClonerEffectorDataInterfaces* UCEClonerEffectorExtension::GetEffectorDataInterfaces() const
-{
-	if (const UCEClonerLayoutBase* LayoutSystem = GetClonerLayout())
-	{
-		return &LayoutSystem->GetDataInterfaces();
-	}
-
-	return nullptr;
-}
-
 void UCEClonerEffectorExtension::OnExtensionActivated()
 {
 	Super::OnExtensionActivated();
 
 	UCEEffectorSubsystem::OnEffectorIdentifierChanged().RemoveAll(this);
 	UCEEffectorSubsystem::OnEffectorIdentifierChanged().AddUObject(this, &UCEClonerEffectorExtension::OnEffectorIdentifierChanged);
-
-	OnEffectorsChanged();
 }
 
 void UCEClonerEffectorExtension::OnExtensionDeactivated()
@@ -94,6 +83,13 @@ void UCEClonerEffectorExtension::OnExtensionDeactivated()
 	UCEEffectorSubsystem::OnEffectorIdentifierChanged().RemoveAll(this);
 }
 
+void UCEClonerEffectorExtension::OnExtensionParametersChanged(UCEClonerComponent* InComponent)
+{
+	Super::OnExtensionParametersChanged(InComponent);
+
+	OnEffectorsChanged();
+}
+
 void UCEClonerEffectorExtension::OnEffectorIdentifierChanged(UCEEffectorComponent* InEffector, int32 InOldIdentifier, int32 InNewIdentifier)
 {
 	if (EffectorActorsWeak.Contains(InEffector->GetOwner()))
@@ -104,9 +100,19 @@ void UCEClonerEffectorExtension::OnEffectorIdentifierChanged(UCEEffectorComponen
 
 void UCEClonerEffectorExtension::OnEffectorsChanged()
 {
-	const FCEClonerEffectorDataInterfaces* EffectorDataInterfaces = GetEffectorDataInterfaces();
+	UNiagaraComponent* Component = GetClonerComponent();
+
+	if (!Component)
+	{
+		return;
+	}
+
+	const FNiagaraUserRedirectionParameterStore& ExposedParameters = Component->GetOverrideParameters();
+
+	static const FNiagaraVariable EffectorIndexDIVar(FNiagaraTypeDefinition(UNiagaraDataInterfaceArrayInt32::StaticClass()), TEXT("EffectorIndexArray"));
+	UNiagaraDataInterfaceArrayInt32* EffectorIndexArrayDI = CastChecked<UNiagaraDataInterfaceArrayInt32>(ExposedParameters.GetDataInterface(EffectorIndexDIVar));
 
-	if (!EffectorDataInterfaces || !EffectorDataInterfaces->GetIndexArray())
+	if (!EffectorIndexArrayDI)
 	{
 		return;
 	}
@@ -179,7 +185,7 @@ void UCEClonerEffectorExtension::OnEffectorsChanged()
 		}
 	}
 
-	TArray<int32>& EffectorIndexArray = EffectorDataInterfaces->GetIndexArray()->GetArrayReference();
+	TArray<int32>& EffectorIndexArray = EffectorIndexArrayDI->GetArrayReference();
 	EffectorIndexArray.Empty(EffectorIndexes.Num());
 
 	for (const int32 EffectorIndex : EffectorIndexes)
diff --git a/Engine/Plugins/Experimental/ClonerEffector/Source/ClonerEffector/Private/Cloner/Extensions/CEClonerEmitterSpawnExtension.cpp b/Engine/Plugins/Experimental/ClonerEffector/Source/ClonerEffector/Private/Cloner/Extensions/CEClonerEmitterSpawnExtension.cpp
--- a/Engine/Plugins/Experimental/ClonerEffector/Source/ClonerEffector/Private/Cloner/Extensions/CEClonerEmitterSpawnExtension.cpp
+++ b/Engine/Plugins/Experimental/ClonerEffector/Source/ClonerEffector/Private/Cloner/Extensions/CEClonerEmitterSpawnExtension.cpp
@@ -3,7 +3,6 @@
 #include "Cloner/Extensions/CEClonerEmitterSpawnExtension.h"
 
 #include "Cloner/CEClonerComponent.h"
-#include "Cloner/Layouts/CEClonerLayoutBase.h"
 #include "NiagaraSystem.h"
 #include "NiagaraUserRedirectionParameterStore.h"
 
@@ -81,22 +80,20 @@ void UCEClonerEmitterSpawnExtension::OnExtensionParametersChanged(UCEClonerCompo
 {
 	Super::OnExtensionParametersChanged(InComponent);
 
-	if (const UCEClonerLayoutBase* ActiveSystem = GetClonerLayout())
-	{
-		FNiagaraUserRedirectionParameterStore& ExposedParameters = ActiveSystem->GetSystem()->GetExposedParameters();
+	FNiagaraUserRedirectionParameterStore& ExposedParameters = InComponent->GetOverrideParameters();
 
-		InComponent->SetFloatParameter(TEXT("SpawnLoopInterval"), SpawnLoopInterval);
+	const FNiagaraVariable SpawnLoopModeVar(FNiagaraTypeDefinition(StaticEnum<ECEClonerSpawnLoopMode>()), TEXT("SpawnLoopMode"));
+	ExposedParameters.SetParameterValue<int32>(static_cast<int32>(SpawnLoopMode), SpawnLoopModeVar);
 
-		InComponent->SetIntParameter(TEXT("SpawnLoopIterations"), SpawnLoopIterations);
+	const ECEClonerSpawnBehaviorMode BehaviorMode = SpawnLoopMode == ECEClonerSpawnLoopMode::Once ? ECEClonerSpawnBehaviorMode::Instant : SpawnBehaviorMode;
+	const FNiagaraVariable SpawnBehaviorModeVar(FNiagaraTypeDefinition(StaticEnum<ECEClonerSpawnBehaviorMode>()), TEXT("SpawnBehaviorMode"));
+	ExposedParameters.SetParameterValue<int32>(static_cast<int32>(BehaviorMode), SpawnBehaviorModeVar);
 
-		InComponent->SetFloatParameter(TEXT("SpawnRate"), SpawnRate);
+	InComponent->SetFloatParameter(TEXT("SpawnLoopInterval"), SpawnLoopInterval);
 
-		const FNiagaraVariable SpawnBehaviorModeVar(FNiagaraTypeDefinition(StaticEnum<ECEClonerSpawnBehaviorMode>()), TEXT("SpawnBehaviorMode"));
-		ExposedParameters.SetParameterValue<int32>(static_cast<int32>(SpawnBehaviorMode), SpawnBehaviorModeVar);
+	InComponent->SetIntParameter(TEXT("SpawnLoopIterations"), SpawnLoopIterations);
 
-		const FNiagaraVariable SpawnLoopModeVar(FNiagaraTypeDefinition(StaticEnum<ECEClonerSpawnLoopMode>()), TEXT("SpawnLoopMode"));
-		ExposedParameters.SetParameterValue<int32>(static_cast<int32>(SpawnLoopMode), SpawnLoopModeVar);
-	}
+	InComponent->SetFloatParameter(TEXT("SpawnRate"), SpawnRate);
 }
 
 #if WITH_EDITOR
diff --git a/Engine/Plugins/Experimental/ClonerEffector/Source/ClonerEffector/Private/Cloner/Extensions/CEClonerExtensionBase.cpp b/Engine/Plugins/Experimental/ClonerEffector/Source/ClonerEffector/Private/Cloner/Extensions/CEClonerExtensionBase.cpp
--- a/Engine/Plugins/Experimental/ClonerEffector/Source/ClonerEffector/Private/Cloner/Extensions/CEClonerExtensionBase.cpp
+++ b/Engine/Plugins/Experimental/ClonerEffector/Source/ClonerEffector/Private/Cloner/Extensions/CEClonerExtensionBase.cpp
@@ -4,6 +4,8 @@
 
 #include "Cloner/CEClonerComponent.h"
 
+DEFINE_LOG_CATEGORY_STATIC(LogCEClonerExtensionBase, Log, All);
+
 #if WITH_EDITOR
 FCEExtensionSection UCEClonerExtensionBase::GetExtensionSection() const
 {
@@ -32,7 +34,7 @@ UCEClonerComponent* UCEClonerExtensionBase::GetClonerComponentChecked() const
 
 UCEClonerLayoutBase* UCEClonerExtensionBase::GetClonerLayout() const
 {
-	if (UCEClonerComponent* ClonerComponent = GetClonerComponent())
+	if (const UCEClonerComponent* ClonerComponent = GetClonerComponent())
 	{
 		return ClonerComponent->GetActiveLayout();
 	}
@@ -45,6 +47,8 @@ void UCEClonerExtensionBase::ActivateExtension()
 	if (!bExtensionActive)
 	{
 		bExtensionActive = true;
+		const UCEClonerComponent* ClonerComponent = GetClonerComponentChecked();
+		UE_LOG(LogCEClonerExtensionBase, Verbose, TEXT("%s : Cloner extension activated %s"), *ClonerComponent->GetOwner()->GetActorNameOrLabel(), *GetExtensionName().ToString());
 		OnExtensionActivated();
 	}
 }
@@ -54,6 +58,8 @@ void UCEClonerExtensionBase::DeactivateExtension()
 	if (bExtensionActive)
 	{
 		bExtensionActive = false;
+		const UCEClonerComponent* ClonerComponent = GetClonerComponentChecked();
+		UE_LOG(LogCEClonerExtensionBase, Verbose, TEXT("%s : Cloner extension deactivated %s"), *ClonerComponent->GetOwner()->GetActorNameOrLabel(), *GetExtensionName().ToString());
 		OnExtensionDeactivated();
 	}
 }
diff --git a/Engine/Plugins/Experimental/ClonerEffector/Source/ClonerEffector/Private/Cloner/Extensions/CEClonerLifetimeExtension.cpp b/Engine/Plugins/Experimental/ClonerEffector/Source/ClonerEffector/Private/Cloner/Extensions/CEClonerLifetimeExtension.cpp
--- a/Engine/Plugins/Experimental/ClonerEffector/Source/ClonerEffector/Private/Cloner/Extensions/CEClonerLifetimeExtension.cpp
+++ b/Engine/Plugins/Experimental/ClonerEffector/Source/ClonerEffector/Private/Cloner/Extensions/CEClonerLifetimeExtension.cpp
@@ -93,29 +93,26 @@ void UCEClonerLifetimeExtension::OnExtensionParametersChanged(UCEClonerComponent
 
 	InComponent->SetBoolParameter(TEXT("LifetimeScaleEnabled"), bLifetimeEnabled && bLifetimeScaleEnabled);
 
-	if (const UCEClonerLayoutBase* ActiveSystem = GetClonerLayout())
-	{
-		FNiagaraUserRedirectionParameterStore& ExposedParameters = ActiveSystem->GetSystem()->GetExposedParameters();
+	FNiagaraUserRedirectionParameterStore& ExposedParameters = InComponent->GetOverrideParameters();
 
-		static const FNiagaraVariable LifetimeScaleCurveVar(FNiagaraTypeDefinition(UNiagaraDataInterfaceCurve::StaticClass()), TEXT("LifetimeScaleCurve"));
+	static const FNiagaraVariable LifetimeScaleCurveVar(FNiagaraTypeDefinition(UNiagaraDataInterfaceCurve::StaticClass()), TEXT("LifetimeScaleCurve"));
 
 #if WITH_EDITOR
-		if (UNiagaraDataInterfaceCurve* LifetimeCurve = LifetimeScaleCurveDIWeak.Get())
-		{
-			LifetimeCurve->OnChanged().RemoveAll(this);
-		}
+	if (UNiagaraDataInterfaceCurve* LifetimeCurve = LifetimeScaleCurveDIWeak.Get())
+	{
+		LifetimeCurve->OnChanged().RemoveAll(this);
+	}
 #endif
 
-		if (UNiagaraDataInterfaceCurve* LifetimeScaleCurveDI = Cast<UNiagaraDataInterfaceCurve>(ExposedParameters.GetDataInterface(LifetimeScaleCurveVar)))
-		{
-			LifetimeScaleCurveDIWeak = LifetimeScaleCurveDI;
-			LifetimeScaleCurveDI->Curve = LifetimeScaleCurve;
+	if (UNiagaraDataInterfaceCurve* LifetimeScaleCurveDI = Cast<UNiagaraDataInterfaceCurve>(ExposedParameters.GetDataInterface(LifetimeScaleCurveVar)))
+	{
+		LifetimeScaleCurveDIWeak = LifetimeScaleCurveDI;
+		LifetimeScaleCurveDI->Curve = LifetimeScaleCurve;
 
 #if WITH_EDITOR
-			LifetimeScaleCurveDI->UpdateLUT();
-			LifetimeScaleCurveDI->OnChanged().AddUObject(this, &UCEClonerLifetimeExtension::OnLifetimeScaleCurveChanged);
+		LifetimeScaleCurveDI->UpdateLUT();
+		LifetimeScaleCurveDI->OnChanged().AddUObject(this, &UCEClonerLifetimeExtension::OnLifetimeScaleCurveChanged);
 #endif
-		}
 	}
 }
 
diff --git a/Engine/Plugins/Experimental/ClonerEffector/Source/ClonerEffector/Private/Cloner/Extensions/CEClonerMeshRendererExtension.cpp b/Engine/Plugins/Experimental/ClonerEffector/Source/ClonerEffector/Private/Cloner/Extensions/CEClonerMeshRendererExtension.cpp
--- a/Engine/Plugins/Experimental/ClonerEffector/Source/ClonerEffector/Private/Cloner/Extensions/CEClonerMeshRendererExtension.cpp
+++ b/Engine/Plugins/Experimental/ClonerEffector/Source/ClonerEffector/Private/Cloner/Extensions/CEClonerMeshRendererExtension.cpp
@@ -118,13 +118,9 @@ void UCEClonerMeshRendererExtension::OnExtensionParametersChanged(UCEClonerCompo
 {
 	Super::OnExtensionParametersChanged(InComponent);
 
-	if (const UCEClonerLayoutBase* CurrentLayout = InComponent->GetClonerActiveLayout())
-	{
-		FNiagaraUserRedirectionParameterStore& ExposedParameters = CurrentLayout->GetSystem()->GetExposedParameters();
-
-		static const FNiagaraVariable MeshModeVar(FNiagaraTypeDefinition(StaticEnum<ECEClonerMeshRenderMode>()), TEXT("MeshRenderMode"));
-		ExposedParameters.SetParameterValue<int32>(static_cast<int32>(MeshRenderMode), MeshModeVar);;
-	}
+	FNiagaraUserRedirectionParameterStore& ExposedParameters = InComponent->GetOverrideParameters();
+	static const FNiagaraVariable MeshModeVar(FNiagaraTypeDefinition(StaticEnum<ECEClonerMeshRenderMode>()), TEXT("MeshRenderMode"));
+	ExposedParameters.SetParameterValue<int32>(static_cast<int32>(MeshRenderMode), MeshModeVar);
 }
 
 int32 UCEClonerMeshRendererExtension::GetClonerMeshesMaterialCount() const
diff --git a/Engine/Plugins/Experimental/ClonerEffector/Source/ClonerEffector/Private/Cloner/Extensions/CEClonerProgressExtension.cpp b/Engine/Plugins/Experimental/ClonerEffector/Source/ClonerEffector/Private/Cloner/Extensions/CEClonerProgressExtension.cpp
--- a/Engine/Plugins/Experimental/ClonerEffector/Source/ClonerEffector/Private/Cloner/Extensions/CEClonerProgressExtension.cpp
+++ b/Engine/Plugins/Experimental/ClonerEffector/Source/ClonerEffector/Private/Cloner/Extensions/CEClonerProgressExtension.cpp
@@ -12,7 +12,7 @@ void UCEClonerProgressExtension::SetInvertProgress(bool bInInvertProgress)
 	}
 
 	bInvertProgress = bInInvertProgress;
-	MarkExtensionDirty(/** UpdateCloner */false);
+	MarkExtensionDirty();
 }
 
 void UCEClonerProgressExtension::SetProgress(float InProgress)
@@ -25,7 +25,7 @@ void UCEClonerProgressExtension::SetProgress(float InProgress)
 	}
 
 	Progress = InProgress;
-	MarkExtensionDirty(/** UpdateCloner */false);
+	MarkExtensionDirty();
 }
 
 void UCEClonerProgressExtension::OnExtensionParametersChanged(UCEClonerComponent* InComponent)
diff --git a/Engine/Plugins/Experimental/ClonerEffector/Source/ClonerEffector/Private/Cloner/Layouts/CEClonerGridLayout.cpp b/Engine/Plugins/Experimental/ClonerEffector/Source/ClonerEffector/Private/Cloner/Layouts/CEClonerGridLayout.cpp
--- a/Engine/Plugins/Experimental/ClonerEffector/Source/ClonerEffector/Private/Cloner/Layouts/CEClonerGridLayout.cpp
+++ b/Engine/Plugins/Experimental/ClonerEffector/Source/ClonerEffector/Private/Cloner/Layouts/CEClonerGridLayout.cpp
@@ -128,7 +128,7 @@ void UCEClonerGridLayout::OnLayoutParametersChanged(UCEClonerComponent* InCompon
 
 	InComponent->SetFloatParameter(TEXT("TwistFactor"), TwistFactor / 100.f);
 
-	FNiagaraUserRedirectionParameterStore& ExposedParameters = InComponent->GetAsset()->GetExposedParameters();
+	FNiagaraUserRedirectionParameterStore& ExposedParameters = InComponent->GetOverrideParameters();
 	static const FNiagaraVariable TwistAxisVar(FNiagaraTypeDefinition(StaticEnum<ENiagaraOrientationAxis>()), TEXT("TwistAxis"));
 	ExposedParameters.SetParameterValue<int32>(static_cast<int32>(TwistAxis), TwistAxisVar);
 }
diff --git a/Engine/Plugins/Experimental/ClonerEffector/Source/ClonerEffector/Private/Cloner/Layouts/CEClonerHoneycombLayout.cpp b/Engine/Plugins/Experimental/ClonerEffector/Source/ClonerEffector/Private/Cloner/Layouts/CEClonerHoneycombLayout.cpp
--- a/Engine/Plugins/Experimental/ClonerEffector/Source/ClonerEffector/Private/Cloner/Layouts/CEClonerHoneycombLayout.cpp
+++ b/Engine/Plugins/Experimental/ClonerEffector/Source/ClonerEffector/Private/Cloner/Layouts/CEClonerHoneycombLayout.cpp
@@ -175,7 +175,7 @@ void UCEClonerHoneycombLayout::OnLayoutParametersChanged(UCEClonerComponent* InC
 
 	InComponent->SetFloatParameter(TEXT("HoneycombHeightSpacing"), HeightSpacing);
 
-	FNiagaraUserRedirectionParameterStore& ExposedParameters = InComponent->GetAsset()->GetExposedParameters();
+	FNiagaraUserRedirectionParameterStore& ExposedParameters = InComponent->GetOverrideParameters();
 	static const FNiagaraVariable HoneycombPlaneVar(FNiagaraTypeDefinition(StaticEnum<ECEClonerPlane>()), TEXT("HoneycombPlane"));
 	ExposedParameters.SetParameterValue<int32>(static_cast<int32>(Plane), HoneycombPlaneVar);
 
diff --git a/Engine/Plugins/Experimental/ClonerEffector/Source/ClonerEffector/Private/Cloner/Layouts/CEClonerLayoutBase.cpp b/Engine/Plugins/Experimental/ClonerEffector/Source/ClonerEffector/Private/Cloner/Layouts/CEClonerLayoutBase.cpp
--- a/Engine/Plugins/Experimental/ClonerEffector/Source/ClonerEffector/Private/Cloner/Layouts/CEClonerLayoutBase.cpp
+++ b/Engine/Plugins/Experimental/ClonerEffector/Source/ClonerEffector/Private/Cloner/Layouts/CEClonerLayoutBase.cpp
@@ -65,7 +65,7 @@ bool UCEClonerLayoutBase::IsLayoutValid() const
 
 bool UCEClonerLayoutBase::IsLayoutLoaded() const
 {
-	return !IsTemplate() && NiagaraSystem && MeshRenderer && DataInterfaces.IsValid();
+	return !IsTemplate() && NiagaraSystem && MeshRenderer;
 }
 
 void UCEClonerLayoutBase::LoadLayout()
@@ -86,7 +86,7 @@ void UCEClonerLayoutBase::LoadLayout()
 		return;
 	}
 
-	UCEClonerComponent* ClonerComponent = GetClonerComponent();
+	const UCEClonerComponent* ClonerComponent = GetClonerComponent();
 
 	if (!IsValid(ClonerComponent))
 	{
@@ -113,7 +113,7 @@ void UCEClonerLayoutBase::LoadLayout()
 	Params.CustomPackageName = FName(TEXT("/") + FString::FromInt(ClonerComponent->GetUniqueID()) + TEXT("_") + GetLayoutName().ToString());
 	Params.CompletionDelegate = MakeUnique<FLoadPackageAsyncDelegate>(FLoadPackageAsyncDelegate::CreateUObject(this, &UCEClonerLayoutBase::OnSystemPackageLoaded));
 
-	UE_LOG(LogCEClonerLayoutBase, Verbose, TEXT("Cloner layout %s : Template system (%s) loading package (%s) async..."), *LayoutName.ToString(), *LayoutAssetPath, *Params.CustomPackageName.ToString())
+	UE_LOG(LogCEClonerLayoutBase, Verbose, TEXT("%s : Cloner layout load requested %s - Template system %s - Package %s"), *GetClonerActor()->GetActorNameOrLabel(), *LayoutName.ToString(), *LayoutAssetPath, *Params.CustomPackageName.ToString())
 
 	LoadRequestIdentifier = LoadPackageAsync(PackagePath, MoveTemp(Params));
 }
@@ -131,10 +131,11 @@ bool UCEClonerLayoutBase::UnloadLayout()
 		return false;
 	}
 
-	DataInterfaces = FCEClonerEffectorDataInterfaces();
 	MeshRenderer = nullptr;
 	NiagaraSystem = nullptr;
 
+	UE_LOG(LogCEClonerLayoutBase, Verbose, TEXT("%s : Cloner layout unloaded %s"), *GetClonerActor()->GetActorNameOrLabel(), *LayoutName.ToString())
+
 	OnLayoutUnloaded();
 
 	return true;
@@ -174,6 +175,8 @@ bool UCEClonerLayoutBase::ActivateLayout()
 
 	ClonerComponent->SetAsset(NiagaraSystem);
 
+	UE_LOG(LogCEClonerLayoutBase, Verbose, TEXT("%s : Cloner layout activated %s"), *GetClonerActor()->GetActorNameOrLabel(), *LayoutName.ToString())
+
 	OnLayoutActive();
 
 	return true;
@@ -195,30 +198,9 @@ bool UCEClonerLayoutBase::DeactivateLayout()
 
 	ClonerComponent->SetAsset(nullptr);
 
-	OnLayoutInactive();
-
-	return true;
-}
-
-bool UCEClonerLayoutBase::CopyTo(UCEClonerLayoutBase* InOtherLayout) const
-{
-	if (!InOtherLayout || InOtherLayout == this)
-	{
-		return false;
-	}
+	UE_LOG(LogCEClonerLayoutBase, Verbose, TEXT("%s : Cloner layout deactivated %s"), *GetClonerActor()->GetActorNameOrLabel(), *LayoutName.ToString())
 
-	if (!IsLayoutLoaded())
-	{
-		return false;
-	}
-
-	const UNiagaraSystem* OtherSystem = InOtherLayout->NiagaraSystem;
-	if (!OtherSystem)
-	{
-		return false;
-	}
-
-	DataInterfaces.CopyTo(InOtherLayout->DataInterfaces);
+	OnLayoutInactive();
 
 	return true;
 }
@@ -310,9 +292,8 @@ void UCEClonerLayoutBase::OnSystemPackageLoaded(const FName& InName, UPackage* I
 #endif
 
 						MeshRenderer = EmitterMeshRenderer;
-						DataInterfaces = FCEClonerEffectorDataInterfaces(NiagaraSystem);
 
-						UE_LOG(LogCEClonerLayoutBase, Log, TEXT("Cloner layout %s : Template system (%s) package (%s) loaded"), *LayoutName.ToString(), *LayoutAssetPath, *InName.ToString())
+						UE_LOG(LogCEClonerLayoutBase, Verbose, TEXT("%s : Cloner layout loaded %s - Template system %s - Package %s"), *GetClonerActor()->GetActorNameOrLabel(), *LayoutName.ToString(), *LayoutAssetPath, *InName.ToString())
 
 						OnLayoutLoaded();
 					}
@@ -322,7 +303,7 @@ void UCEClonerLayoutBase::OnSystemPackageLoaded(const FName& InName, UPackage* I
 	}
 	else
 	{
-		UE_LOG(LogCEClonerLayoutBase, Warning, TEXT("Cloner layout %s : Template system (%s) package (%s) could not be async loaded"), *LayoutName.ToString(), *LayoutAssetPath, *InName.ToString())
+		UE_LOG(LogCEClonerLayoutBase, Warning, TEXT("%s : Cloner layout load failed %s - Template system %s - Package %s"), *GetClonerActor()->GetActorNameOrLabel(), *LayoutName.ToString(), *LayoutAssetPath, *InName.ToString())
 	}
 
 	OnClonerLayoutLoadedDelegate.Broadcast(this, IsLayoutLoaded());
@@ -336,7 +317,7 @@ UCEClonerComponent* UCEClonerLayoutBase::GetClonerComponent() const
 
 AActor* UCEClonerLayoutBase::GetClonerActor() const
 {
-	if (UCEClonerComponent* ClonerComponent = GetClonerComponent())
+	if (const UCEClonerComponent* ClonerComponent = GetClonerComponent())
 	{
 		return ClonerComponent->GetOwner();
 	}
diff --git a/Engine/Plugins/Experimental/ClonerEffector/Source/ClonerEffector/Private/Cloner/Layouts/CEClonerMeshLayout.cpp b/Engine/Plugins/Experimental/ClonerEffector/Source/ClonerEffector/Private/Cloner/Layouts/CEClonerMeshLayout.cpp
--- a/Engine/Plugins/Experimental/ClonerEffector/Source/ClonerEffector/Private/Cloner/Layouts/CEClonerMeshLayout.cpp
+++ b/Engine/Plugins/Experimental/ClonerEffector/Source/ClonerEffector/Private/Cloner/Layouts/CEClonerMeshLayout.cpp
@@ -104,7 +104,7 @@ void UCEClonerMeshLayout::OnLayoutParametersChanged(UCEClonerComponent* InCompon
 
 	InComponent->SetIntParameter(TEXT("SampleMeshCount"), Count);
 
-	FNiagaraUserRedirectionParameterStore& ExposedParameters = InComponent->GetAsset()->GetExposedParameters();
+	FNiagaraUserRedirectionParameterStore& ExposedParameters = InComponent->GetOverrideParameters();
 
 	static const FNiagaraVariable SampleMeshAssetVar(FNiagaraTypeDefinition(StaticEnum<ECEClonerMeshAsset>()), TEXT("SampleMeshAsset"));
 	ExposedParameters.SetParameterValue<int32>(static_cast<int32>(Asset), SampleMeshAssetVar);
diff --git a/Engine/Plugins/Experimental/ClonerEffector/Source/ClonerEffector/Private/Cloner/Layouts/CEClonerSplineLayout.cpp b/Engine/Plugins/Experimental/ClonerEffector/Source/ClonerEffector/Private/Cloner/Layouts/CEClonerSplineLayout.cpp
--- a/Engine/Plugins/Experimental/ClonerEffector/Source/ClonerEffector/Private/Cloner/Layouts/CEClonerSplineLayout.cpp
+++ b/Engine/Plugins/Experimental/ClonerEffector/Source/ClonerEffector/Private/Cloner/Layouts/CEClonerSplineLayout.cpp
@@ -140,7 +140,7 @@ void UCEClonerSplineLayout::OnLayoutParametersChanged(UCEClonerComponent* InComp
 
 	InComponent->SetIntParameter(TEXT("SampleSplineCount"), Count);
 
-	const FNiagaraUserRedirectionParameterStore& ExposedParameters = InComponent->GetAsset()->GetExposedParameters();
+	const FNiagaraUserRedirectionParameterStore& ExposedParameters = InComponent->GetOverrideParameters();
 	static const FNiagaraVariable SampleSplineVar(FNiagaraTypeDefinition(UNiagaraDataInterfaceSpline::StaticClass()), TEXT("SampleSpline"));
 	UNiagaraDataInterfaceSpline* SplineDI = Cast<UNiagaraDataInterfaceSpline>(ExposedParameters.GetDataInterface(SampleSplineVar));
 
diff --git a/Engine/Plugins/Experimental/ClonerEffector/Source/ClonerEffector/Public/CEClonerEffectorShared.h b/Engine/Plugins/Experimental/ClonerEffector/Source/ClonerEffector/Public/CEClonerEffectorShared.h
--- a/Engine/Plugins/Experimental/ClonerEffector/Source/ClonerEffector/Public/CEClonerEffectorShared.h
+++ b/Engine/Plugins/Experimental/ClonerEffector/Source/ClonerEffector/Public/CEClonerEffectorShared.h
@@ -5,7 +5,6 @@
 #include "Components/SceneComponent.h"
 #include "Engine/StaticMesh.h"
 #include "Materials/MaterialInterface.h"
-#include "NiagaraDataInterfaceArrayInt.h"
 #include "CEClonerEffectorShared.generated.h"
 
 class AActor;
@@ -449,33 +448,6 @@ protected:
 	void Read(const UNiagaraDataChannelReader* InReader);
 };
 
-USTRUCT()
-struct FCEClonerEffectorDataInterfaces
-{
-	friend class UCEClonerLayoutBase;
-
-	GENERATED_BODY()
-
-	static inline const FName IndexName = TEXT("EffectorIndexArray");
-
-	explicit FCEClonerEffectorDataInterfaces(const UNiagaraSystem* InSystem);
-	FCEClonerEffectorDataInterfaces() = default;
-
-	void Clear() const;
-	void CopyTo(FCEClonerEffectorDataInterfaces& InOther) const;
-	void Resize(int32 InSize) const;
-	void Remove(int32 InIndex) const;
-	bool IsValid() const;
-	int32 Num() const;
-	void Commit() const;
-
-	UNiagaraDataInterfaceArrayInt32* GetIndexArray() const;
-
-protected:
-	UPROPERTY()
-	TMap<FName, TObjectPtr<UNiagaraDataInterface>> DataInterfaces;
-};
-
 USTRUCT(BlueprintType)
 struct CLONEREFFECTOR_API FCEClonerGridConstraintSphere
 {
diff --git a/Engine/Plugins/Experimental/ClonerEffector/Source/ClonerEffector/Public/Cloner/CEClonerComponent.h b/Engine/Plugins/Experimental/ClonerEffector/Source/ClonerEffector/Public/Cloner/CEClonerComponent.h
--- a/Engine/Plugins/Experimental/ClonerEffector/Source/ClonerEffector/Public/Cloner/CEClonerComponent.h
+++ b/Engine/Plugins/Experimental/ClonerEffector/Source/ClonerEffector/Public/Cloner/CEClonerComponent.h
@@ -238,24 +238,6 @@ protected:
 	void UpdateAttachmentTree();
 	void UpdateActorAttachment(AActor* InActor, AActor* InParent);
 
-	void BindActorDelegates(AActor* InActor);
-	void UnbindActorDelegates(AActor* InActor) const;
-
-	void OnTransformUpdated(USceneComponent* InUpdatedComponent, EUpdateTransformFlags InUpdateTransformFlags, ETeleportType InTeleport);
-
-	UFUNCTION()
-	void OnActorDestroyed(AActor* InDestroyedActor);
-
-#if WITH_EDITOR
-	void OnActorPropertyChanged(UObject* InObject, FPropertyChangedEvent& InPropertyChangedEvent);
-#endif
-
-	/** Called when a material has changed in any cloned actors */
-	void OnMaterialChanged(UObject* InObject);
-
-	/** Called when a mesh has been updated and cloner needs to update instances */
-	void OnMeshChanged(UStaticMeshComponent*, AActor* InActor);
-
 	/** Get cloner direct children in the correct order */
 	void GetOrderedRootActors(TArray<AActor*>& OutActors) const;
 
@@ -365,6 +347,31 @@ private:
 	void OnVisualizerSpriteVisibleChanged();
 #endif
 
+	void BindActorDelegates(AActor* InActor);
+	void UnbindActorDelegates(AActor* InActor) const;
+
+	void SetActorVisibility(AActor* InActor, bool bInVisibility);
+
+#if WITH_EDITOR
+	void OnActorPropertyChanged(UObject* InObject, FPropertyChangedEvent& InPropertyChangedEvent);
+#endif
+
+	/** Called when a cloned actor is destroyed */
+	UFUNCTION()
+	void OnActorDestroyed(AActor* InDestroyedActor);
+
+	/** Called when a material has changed in any cloned actors */
+	void OnMaterialChanged(UObject* InObject);
+
+	/** Called when a mesh has changed in any cloned actors */
+	void OnMeshChanged(UStaticMeshComponent*, AActor* InActor);
+
+	/** Called when the render state of a cloned component changes */
+	void OnRenderStateDirty(UActorComponent& InActorComponent);
+
+	/** Called when the transform state of a cloned component changes */
+	void OnComponentTransformed(USceneComponent* InComponent, EUpdateTransformFlags InFlags, ETeleportType InTeleport);
+
 	template<
 		typename InLayoutClass
 		UE_REQUIRES(TIsDerivedFrom<InLayoutClass, UCEClonerLayoutBase>::Value)>
@@ -419,6 +426,8 @@ private:
 	FTSTicker::FDelegateHandle ClonerTickerHandle;
 
 #if WITH_EDITOR
+	double LastNotificationTime = 0.0;
+
 	/** Used for PECP */
 	static const TCEPropertyChangeDispatcher<UCEClonerComponent> PropertyChangeDispatcher;
 #endif
diff --git a/Engine/Plugins/Experimental/ClonerEffector/Source/ClonerEffector/Public/Cloner/CEClonerMeshBuilder.h b/Engine/Plugins/Experimental/ClonerEffector/Source/ClonerEffector/Public/Cloner/CEClonerMeshBuilder.h
--- a/Engine/Plugins/Experimental/ClonerEffector/Source/ClonerEffector/Public/Cloner/CEClonerMeshBuilder.h
+++ b/Engine/Plugins/Experimental/ClonerEffector/Source/ClonerEffector/Public/Cloner/CEClonerMeshBuilder.h
@@ -6,6 +6,8 @@
 #include "UObject/WeakObjectPtrTemplates.h"
 #include "CEClonerMeshBuilder.generated.h"
 
+class AActor;
+class UActorComponent;
 class UBrushComponent;
 class UDynamicMesh;
 class UDynamicMeshComponent;
@@ -41,6 +43,15 @@ struct FCEClonerMeshBuilder
 
 	static const FCEClonerMeshBuilderParams DefaultParams;
 
+	/** Checks if the component contains any geometry data */
+	static bool HasAnyGeometry(UActorComponent* InComponent);
+
+	/** Does the mesh builder supports this actor */
+	static bool IsActorSupported(const AActor* InActor);
+
+	/** Does the mesh builder supports this component */
+	static bool IsComponentSupported(const UActorComponent* InComponent);
+
 	FCEClonerMeshBuilder();
 
 	int32 GetMeshInstanceCount() const
diff --git a/Engine/Plugins/Experimental/ClonerEffector/Source/ClonerEffector/Public/Cloner/Extensions/CEClonerEffectorExtension.h b/Engine/Plugins/Experimental/ClonerEffector/Source/ClonerEffector/Public/Cloner/Extensions/CEClonerEffectorExtension.h
--- a/Engine/Plugins/Experimental/ClonerEffector/Source/ClonerEffector/Public/Cloner/Extensions/CEClonerEffectorExtension.h
+++ b/Engine/Plugins/Experimental/ClonerEffector/Source/ClonerEffector/Public/Cloner/Extensions/CEClonerEffectorExtension.h
@@ -42,9 +42,6 @@ public:
 	/** Checks if an effector is linked with this cloner */
 	CLONEREFFECTOR_API bool IsEffectorLinked(const AActor* InEffectorActor) const;
 
-	/** Used by effector actors to apply transformations to this cloner instances */
-	const FCEClonerEffectorDataInterfaces* GetEffectorDataInterfaces() const;
-
 #if WITH_EDITOR
 	/** This will create a new effector actor, link it to this cloner and select it */
 	UFUNCTION(CallInEditor, Category="Effector")
@@ -62,6 +59,7 @@ protected:
 	//~ Begin UCEClonerExtensionBase
 	virtual void OnExtensionActivated() override;
 	virtual void OnExtensionDeactivated() override;
+	virtual void OnExtensionParametersChanged(UCEClonerComponent* InComponent) override;
 	//~ End UCEClonerExtensionBase
 
 	void OnEffectorIdentifierChanged(UCEEffectorComponent* InEffector, int32 InOldIdentifier, int32 InNewIdentifier);
diff --git a/Engine/Plugins/Experimental/ClonerEffector/Source/ClonerEffector/Public/Cloner/Layouts/CEClonerLayoutBase.h b/Engine/Plugins/Experimental/ClonerEffector/Source/ClonerEffector/Public/Cloner/Layouts/CEClonerLayoutBase.h
--- a/Engine/Plugins/Experimental/ClonerEffector/Source/ClonerEffector/Public/Cloner/Layouts/CEClonerLayoutBase.h
+++ b/Engine/Plugins/Experimental/ClonerEffector/Source/ClonerEffector/Public/Cloner/Layouts/CEClonerLayoutBase.h
@@ -58,11 +58,6 @@ public:
 		return MeshRenderer;
 	}
 
-	const FCEClonerEffectorDataInterfaces& GetDataInterfaces() const
-	{
-		return DataInterfaces;
-	}
-
 	TMulticastDelegateRegistration<void(UCEClonerLayoutBase*, bool)>& OnLayoutLoadedDelegate()
 	{
 		return OnClonerLayoutLoadedDelegate;
@@ -100,9 +95,6 @@ public:
 	/** Deactivate this layout system if active */
 	bool DeactivateLayout();
 
-	/** Copies this layout data interfaces to other layout */
-	bool CopyTo(UCEClonerLayoutBase* InOtherLayout) const;
-
 	/** Gets the cloner extensions supported by this layout */
 	TSet<TSubclassOf<UCEClonerExtensionBase>> GetSupportedExtensions() const;
 
@@ -165,10 +157,6 @@ private:
 	UPROPERTY(Transient)
 	TObjectPtr<UNiagaraMeshRendererProperties> MeshRenderer;
 
-	/** Data interfaces used by the effectors */
-	UPROPERTY(Transient)
-	FCEClonerEffectorDataInterfaces DataInterfaces;
-
 	int32 LoadRequestIdentifier = INDEX_NONE;
 
 	ECEClonerSystemStatus LayoutStatus = ECEClonerSystemStatus::UpToDate;
diff --git a/Engine/Plugins/Experimental/Text3D/Source/Text3D/Private/Text3DComponent.cpp b/Engine/Plugins/Experimental/Text3D/Source/Text3D/Private/Text3DComponent.cpp
--- a/Engine/Plugins/Experimental/Text3D/Source/Text3D/Private/Text3DComponent.cpp
+++ b/Engine/Plugins/Experimental/Text3D/Source/Text3D/Private/Text3DComponent.cpp
@@ -24,30 +24,6 @@
 
 #define LOCTEXT_NAMESPACE "Text3D"
 
-#if WITH_EDITOR
-namespace UE::Text3D::Private
-{
-	/** Get the Group Type according to the material name. */
-	static TMap<FName, EText3DGroupType> MaterialToGroup =
-	{
-		{ TEXT("FrontMaterial"), EText3DGroupType::Front },
-		{ TEXT("BackMaterial"), EText3DGroupType::Back },
-		{ TEXT("ExtrudeMaterial"), EText3DGroupType::Extrude },
-		{ TEXT("BevelMaterial"), EText3DGroupType::Bevel }
-	};
-
-	/** Sets the material based on the group name. */
-	using FGetter = TMemFunPtrType<true, UText3DComponent, UMaterialInterface*()>::Type;
-	static TMap<FName, FGetter> GroupToMaterial =
-	{
-		{ TEXT("FrontMaterial"), &UText3DComponent::GetFrontMaterial },
-		{ TEXT("BackMaterial"), &UText3DComponent::GetBackMaterial },
-		{ TEXT("ExtrudeMaterial"), &UText3DComponent::GetExtrudeMaterial },
-		{ TEXT("BevelMaterial"), &UText3DComponent::GetBevelMaterial }
-	};
-}
-#endif
-
 struct FText3DShapedText
 {
 	FText3DShapedText()
@@ -327,9 +303,12 @@ void UText3DComponent::PostEditChangeProperty(FPropertyChangedEvent& PropertyCha
 	{
 		MarkForLayoutUpdate();
 	}
-	else if (const EText3DGroupType* MaterialGroup = UE::Text3D::Private::MaterialToGroup.Find(Name))
+	else if (Name == GET_MEMBER_NAME_CHECKED(UText3DComponent, FrontMaterial)
+		|| Name == GET_MEMBER_NAME_CHECKED(UText3DComponent, BevelMaterial)
+		|| Name == GET_MEMBER_NAME_CHECKED(UText3DComponent, ExtrudeMaterial)
+		|| Name == GET_MEMBER_NAME_CHECKED(UText3DComponent, BackMaterial))
 	{
-		UpdateMaterial(*MaterialGroup, Invoke(UE::Text3D::Private::GroupToMaterial[Name], this));
+		OnMaterialChanged();
 	}
 
 	RebuildInternal();
@@ -665,7 +644,7 @@ void UText3DComponent::SetMaterial(const EText3DGroupType Type, UMaterialInterfa
 		}
 		}
 
-		UpdateMaterial(Type, Value);
+		OnMaterialChanged();
 	}
 }
 
@@ -1302,11 +1281,7 @@ void UText3DComponent::BuildTextMeshInternal(const bool& bCleanCache)
 		}
 	}
 
-	for (int32 Index = 0; Index < static_cast<int32>(EText3DGroupType::TypeCount); Index++)
-	{
-		const EText3DGroupType Type = static_cast<EText3DGroupType>(Index);
-		UpdateMaterial(Type, GetMaterial(Type));
-	}
+	OnMaterialChanged();
 
 	TextGeneratedNativeDelegate.Broadcast();
 	TextGeneratedDelegate.Broadcast();
@@ -1332,28 +1307,47 @@ float UText3DComponent::MaxBevel() const
 	return Extrude / 2.0f;
 }
 
-void UText3DComponent::UpdateMaterial(const EText3DGroupType Type, UMaterialInterface* Material)
+void UText3DComponent::OnMaterialChanged()
 {
-	// Material indices are affected by some options
-	const bool bHasBevel = !bOutline && !FMath::IsNearlyZero(Bevel);
-	if (!bHasBevel && Type == EText3DGroupType::Bevel)
+	// Material indices are affected by some options and can differ
+	auto GetMaterialTypeIndex = [this](EText3DGroupType InType)->int32
 	{
-		return; 
-	}
+		const bool bHasBevel = !bOutline && !FMath::IsNearlyZero(Bevel);
 
-	const bool bHasExtrude = !FMath::IsNearlyZero(Extrude);
-	if (!bHasExtrude && Type == EText3DGroupType::Extrude)
-	{
-		return;
-	}
+		int32 MaterialIndex = static_cast<int32>(InType);
 
-	int32 Index = static_cast<int32>(Type);
-	Index -= !bHasBevel && Type >= EText3DGroupType::Bevel ? 1 : 0; // if no bevel, and the input is bevel or above (bevel, side/extrude, back), offset -1
-	Index -= !bHasExtrude && Type >= EText3DGroupType::Extrude ? 1 : 0; // if no extrude, and the input is side/extrude or above (back), offset -1
+		if (InType >= EText3DGroupType::Extrude && !bHasBevel)
+		{
+			MaterialIndex -= 1;
+		}
 
+		return MaterialIndex;
+	};
+
+	/**
+	* GetNumMaterials should return 4 but the index of the material changes based on options used
+	* [Front, (Extrude|Bevel), Back] if no extrude or bevel
+	* [Front, Extrude, Back] if no bevel since bevel only works when there is extrude anyway
+	* [Front, Bevel, Extrude, Back] if extrude and bevel is set
+	 */
 	for (UStaticMeshComponent* StaticMeshComponent : CharacterMeshes)
 	{
-		StaticMeshComponent->SetMaterial(Index, Material);
+		const int32 MaterialCount = StaticMeshComponent->GetNumMaterials();
+
+		for (int32 MaterialIndex = 0; MaterialIndex < MaterialCount; MaterialIndex++)
+		{
+			// Should not have any material slot above the back material slot but in case set them as nullptr
+			if (MaterialIndex >= static_cast<int32>(EText3DGroupType::TypeCount))
+			{
+				StaticMeshComponent->SetMaterial(MaterialIndex, nullptr);
+				continue;
+			}
+
+			// Get the material type from the current index and based on text options adapt the index to assign proper material
+			const EText3DGroupType MaterialType = static_cast<EText3DGroupType>(MaterialIndex);
+			const int32 MaterialTypeIndex = GetMaterialTypeIndex(MaterialType);
+			StaticMeshComponent->SetMaterial(MaterialTypeIndex, GetMaterial(MaterialType));
+		}
 	}
 }
 
diff --git a/Engine/Plugins/Experimental/Text3D/Source/Text3D/Public/Text3DComponent.h b/Engine/Plugins/Experimental/Text3D/Source/Text3D/Public/Text3DComponent.h
--- a/Engine/Plugins/Experimental/Text3D/Source/Text3D/Public/Text3DComponent.h
+++ b/Engine/Plugins/Experimental/Text3D/Source/Text3D/Public/Text3DComponent.h
@@ -464,7 +464,7 @@ private:
 	/** Slot based Material accessors. */
 	UMaterialInterface* GetMaterial(const EText3DGroupType Type) const;
 	void SetMaterial(const EText3DGroupType Type, UMaterialInterface* Material);
-	void UpdateMaterial(const EText3DGroupType Type, UMaterialInterface* Material);
+	void OnMaterialChanged();
 
 	/** Validation and context sensitive limits applied to the current bevel value. */
 	void CheckBevel();
