MotionDesign : ClonerEffector
- Fixed GC memory leaks with niagara simulation when switching levels and system is still compiling
- Improving cloner duplication speed by avoiding duplicating some properties

diff --git a/Engine/Plugins/Experimental/ClonerEffector/Source/ClonerEffector/Private/Cloner/CEClonerComponent.cpp b/Engine/Plugins/Experimental/ClonerEffector/Source/ClonerEffector/Private/Cloner/CEClonerComponent.cpp
--- a/Engine/Plugins/Experimental/ClonerEffector/Source/ClonerEffector/Private/Cloner/CEClonerComponent.cpp
+++ b/Engine/Plugins/Experimental/ClonerEffector/Source/ClonerEffector/Private/Cloner/CEClonerComponent.cpp
@@ -378,10 +378,7 @@ void UCEClonerComponent::BindActorDelegates(AActor* InActor)
 		return;
 	}
 
-	if (!InActor->OnDestroyed.IsAlreadyBound(this, &UCEClonerComponent::OnActorDestroyed))
-	{
-		InActor->OnDestroyed.AddDynamic(this, &UCEClonerComponent::OnActorDestroyed);
-	}
+	InActor->OnDestroyed.AddUniqueDynamic(this, &UCEClonerComponent::OnActorDestroyed);
 
 #if WITH_EDITOR
 	// Detect static mesh change
diff --git a/Engine/Plugins/Experimental/ClonerEffector/Source/ClonerEffector/Private/Cloner/Layouts/CEClonerLayoutBase.cpp b/Engine/Plugins/Experimental/ClonerEffector/Source/ClonerEffector/Private/Cloner/Layouts/CEClonerLayoutBase.cpp
--- a/Engine/Plugins/Experimental/ClonerEffector/Source/ClonerEffector/Private/Cloner/Layouts/CEClonerLayoutBase.cpp
+++ b/Engine/Plugins/Experimental/ClonerEffector/Source/ClonerEffector/Private/Cloner/Layouts/CEClonerLayoutBase.cpp
@@ -5,6 +5,8 @@
 #include "Cloner/CEClonerActor.h"
 #include "Cloner/CEClonerComponent.h"
 #include "Cloner/Extensions/CEClonerExtensionBase.h"
+#include "Engine/Level.h"
+#include "Engine/World.h"
 #include "Misc/PackageName.h"
 #include "NiagaraEmitter.h"
 #include "NiagaraMeshRendererProperties.h"
@@ -94,28 +96,39 @@ void UCEClonerLayoutBase::LoadLayout()
 	}
 
 	// Extract package path
-	FString PackagePath = LayoutAssetPath;
+	FString MountedPath = LayoutAssetPath;
+	{
+		int32 FirstQuoteIndex;
+		MountedPath.FindChar('\'', FirstQuoteIndex);
 
-	int32 FirstQuoteIndex;
-	PackagePath.FindChar('\'', FirstQuoteIndex);
+		int32 LastQuoteIndex;
+		MountedPath.FindLastChar('\'', LastQuoteIndex);
 
-	int32 LastQuoteIndex;
-	PackagePath.FindLastChar('\'', LastQuoteIndex);
+		if (FirstQuoteIndex != LastQuoteIndex)
+		{
+			MountedPath = MountedPath.Mid(FirstQuoteIndex + 1, LastQuoteIndex - FirstQuoteIndex - 1);
+		}
 
-	if (FirstQuoteIndex != LastQuoteIndex)
-	{
-		PackagePath = PackagePath.Mid(FirstQuoteIndex + 1, LastQuoteIndex - FirstQuoteIndex - 1);
+		MountedPath = FPackageName::ObjectPathToPackageName(MountedPath);
 	}
 
-	PackagePath = FPackageName::ObjectPathToPackageName(PackagePath);
+	FPackagePath LayoutPackagePath;
+	FPackagePath::TryFromMountedName(MountedPath, LayoutPackagePath);
+
+	FPackagePath CustomPackagePath;
+	FPackagePath::TryFromPackageName(TEXT("/Game/Temp/") + GetLayoutName().ToString() + TEXT("_") + FString::FromInt(ClonerComponent->GetUniqueID()), CustomPackagePath);
+	CustomPackagePath.SetHeaderExtension(EPackageExtension::Asset);
 
 	FLoadPackageAsyncOptionalParams Params;
-	Params.CustomPackageName = FName(TEXT("/") + FString::FromInt(ClonerComponent->GetUniqueID()) + TEXT("_") + GetLayoutName().ToString());
+	Params.PackagePriority = INT32_MAX;
+	Params.CustomPackageName = CustomPackagePath.GetPackageFName();
 	Params.CompletionDelegate = MakeUnique<FLoadPackageAsyncDelegate>(FLoadPackageAsyncDelegate::CreateUObject(this, &UCEClonerLayoutBase::OnSystemPackageLoaded));
 
 	UE_LOG(LogCEClonerLayoutBase, Verbose, TEXT("%s : Cloner layout load requested %s - Template system %s - Package %s"), *GetClonerActor()->GetActorNameOrLabel(), *LayoutName.ToString(), *LayoutAssetPath, *Params.CustomPackageName.ToString())
 
-	LoadRequestIdentifier = LoadPackageAsync(PackagePath, MoveTemp(Params));
+	LoadRequestIdentifier = LoadPackageAsync(LayoutPackagePath, MoveTemp(Params));
+
+	BindCleanupDelegates();
 }
 
 bool UCEClonerLayoutBase::UnloadLayout()
@@ -131,6 +144,15 @@ bool UCEClonerLayoutBase::UnloadLayout()
 		return false;
 	}
 
+	MeshRenderer->Meshes.Empty();
+	NiagaraSystem->RemoveFromRoot();
+#if WITH_EDITOR
+	NiagaraSystem->KillAllActiveCompilations();
+#endif
+	UPackage* Package = NiagaraSystem->GetPackage();
+	Package->ClearFlags(RF_Standalone);
+	Package->MarkAsGarbage();
+
 	MeshRenderer = nullptr;
 	NiagaraSystem = nullptr;
 
@@ -196,6 +218,7 @@ bool UCEClonerLayoutBase::DeactivateLayout()
 		return false;
 	}
 
+	ClonerComponent->GetOverrideParameters().Empty(/** ClearBindings */true);
 	ClonerComponent->SetAsset(nullptr);
 
 	UE_LOG(LogCEClonerLayoutBase, Verbose, TEXT("%s : Cloner layout deactivated %s"), *GetClonerActor()->GetActorNameOrLabel(), *LayoutName.ToString())
@@ -272,11 +295,8 @@ void UCEClonerLayoutBase::OnSystemPackageLoaded(const FName& InName, UPackage* I
 
 	if (NiagaraSystem)
 	{
-		// Change outer to avoid GC leak
 		InPackage->SetFlags(RF_Transient);
 		NiagaraSystem->RemoveFromRoot();
-		NiagaraSystem->Rename(NULL, this, REN_DoNotDirty | REN_DontCreateRedirectors | REN_NonTransactional);
-		InPackage->MarkAsGarbage();
 
 		for (FNiagaraEmitterHandle& SystemEmitterHandle : NiagaraSystem->GetEmitterHandles())
 		{
@@ -310,6 +330,56 @@ void UCEClonerLayoutBase::OnSystemPackageLoaded(const FName& InName, UPackage* I
 	OnClonerLayoutLoadedDelegate.Clear();
 }
 
+void UCEClonerLayoutBase::BindCleanupDelegates()
+{
+	UnbindCleanupDelegates();
+
+	if (const UCEClonerComponent* ClonerComponent = GetClonerComponent())
+	{
+		if (ULevel* ClonerLevel = ClonerComponent->GetComponentLevel())
+		{
+			ClonerLevel->OnCleanupLevel.AddUObject(this, &UCEClonerLayoutBase::OnLevelCleanup);
+		}
+
+		FWorldDelegates::OnWorldCleanup.AddUObject(this, &UCEClonerLayoutBase::OnWorldCleanup);
+	}
+}
+
+void UCEClonerLayoutBase::UnbindCleanupDelegates() const
+{
+	if (const UCEClonerComponent* ClonerComponent = GetClonerComponent())
+	{
+		if (ULevel* ClonerLevel = ClonerComponent->GetComponentLevel())
+		{
+			ClonerLevel->OnCleanupLevel.RemoveAll(this);
+		}
+
+		FWorldDelegates::OnWorldCleanup.RemoveAll(this);
+	}
+}
+
+void UCEClonerLayoutBase::OnWorldCleanup(UWorld* InWorld, bool bInSessionEnded, bool bInCleanupResources)
+{
+	const AActor* Actor = GetClonerActor();
+	if (bInCleanupResources && Actor && Actor->GetWorld() == InWorld)
+	{
+		OnLevelCleanup();
+	}
+}
+
+void UCEClonerLayoutBase::OnLevelCleanup()
+{
+	if (IsLayoutLoaded())
+	{
+		UE_LOG(LogCEClonerLayoutBase, Log, TEXT("%s : Cloner layout cleanup %s"), *GetClonerActor()->GetActorNameOrLabel(), *LayoutName.ToString())
+
+		DeactivateLayout();
+		UnloadLayout();
+	}
+
+	UnbindCleanupDelegates();
+}
+
 UCEClonerComponent* UCEClonerLayoutBase::GetClonerComponent() const
 {
 	return GetTypedOuter<UCEClonerComponent>();
diff --git a/Engine/Plugins/Experimental/ClonerEffector/Source/ClonerEffector/Private/Subsystems/CEEffectorSubsystem.cpp b/Engine/Plugins/Experimental/ClonerEffector/Source/ClonerEffector/Private/Subsystems/CEEffectorSubsystem.cpp
--- a/Engine/Plugins/Experimental/ClonerEffector/Source/ClonerEffector/Private/Subsystems/CEEffectorSubsystem.cpp
+++ b/Engine/Plugins/Experimental/ClonerEffector/Source/ClonerEffector/Private/Subsystems/CEEffectorSubsystem.cpp
@@ -293,12 +293,10 @@ void UCEEffectorSubsystem::SetLevelEffectorsEnabled(const UWorld* InWorld, bool
 
 void UCEEffectorSubsystem::UpdateEffectorChannel(const UWorld* InWorld)
 {
-	if (!IsValid(InWorld))
-	{
-		return;
-	}
-
-	if (EffectorsWeak.IsEmpty())
+	if (!IsValid(InWorld)
+		|| !InWorld->IsInitialized()
+		|| InWorld->IsBeingCleanedUp()
+		|| EffectorsWeak.IsEmpty())
 	{
 		return;
 	}
@@ -315,17 +313,12 @@ void UCEEffectorSubsystem::UpdateEffectorChannel(const UWorld* InWorld)
 
 	// Remove invalid effectors and push updates to effector assigned channel indexes
 	int32 EffectorIndex = 0;
-	for (TArray<TWeakObjectPtr<UCEEffectorComponent>>::TIterator It(EffectorsWeak); It; ++It)
+	for (const TWeakObjectPtr<UCEEffectorComponent>& EffectorWeak : EffectorsWeak)
 	{
-		UCEEffectorComponent* Effector = It->Get();
-
-		if (!IsValid(Effector))
-		{
-			It.RemoveCurrent();
-			continue;
-		}
+		UCEEffectorComponent* Effector = EffectorWeak.Get();
 
-		if (Effector->GetWorld() != InWorld)
+		if (!IsValid(Effector)
+			|| Effector->GetWorld() != InWorld)
 		{
 			continue;
 		}
@@ -367,19 +360,27 @@ TStatId UCEEffectorSubsystem::GetStatId() const
 
 void UCEEffectorSubsystem::Tick(float InDeltaTime)
 {
-	TSet<UWorld*> Worlds;
+	TSet<const UWorld*> Worlds;
 
-	for (const TWeakObjectPtr<UCEEffectorComponent>& EffectorWeak : EffectorsWeak)
+	for (TArray<TWeakObjectPtr<UCEEffectorComponent>>::TIterator It(EffectorsWeak); It; ++It)
 	{
-		UCEEffectorComponent* Effector = EffectorWeak.Get();
+		const UCEEffectorComponent* Effector = It->Get();
+		if (!IsValid(Effector))
+		{
+			It.RemoveCurrent();
+			continue;
+		}
 
-		if (Effector && Effector->GetWorld())
+		const UWorld* EffectorWorld = Effector->GetWorld();
+		if (!IsValid(EffectorWorld) || !EffectorWorld->IsInitialized() || EffectorWorld->IsBeingCleanedUp())
 		{
-			Worlds.Add(Effector->GetWorld());
+			continue;
 		}
+
+		Worlds.Add(EffectorWorld);
 	}
 
-	for (UWorld* World : Worlds)
+	for (const UWorld* World : Worlds)
 	{
 		UpdateEffectorChannel(World);
 	}
diff --git a/Engine/Plugins/Experimental/ClonerEffector/Source/ClonerEffector/Public/Cloner/CEClonerComponent.h b/Engine/Plugins/Experimental/ClonerEffector/Source/ClonerEffector/Public/Cloner/CEClonerComponent.h
--- a/Engine/Plugins/Experimental/ClonerEffector/Source/ClonerEffector/Public/Cloner/CEClonerComponent.h
+++ b/Engine/Plugins/Experimental/ClonerEffector/Source/ClonerEffector/Public/Cloner/CEClonerComponent.h
@@ -405,10 +405,10 @@ private:
 	TArray<FName> GetClonerLayoutNames() const;
 
 	/** Attachment tree view */
-	UPROPERTY(Transient, NonTransactional)
+	UPROPERTY(Transient, DuplicateTransient, NonTransactional)
 	FCEClonerAttachmentTree ClonerTree;
 
-	UPROPERTY(Transient, NonTransactional)
+	UPROPERTY(Transient, DuplicateTransient, NonTransactional)
 	FCEClonerMeshBuilder MeshBuilder;
 
 	/** Asset meshes needs update */
diff --git a/Engine/Plugins/Experimental/ClonerEffector/Source/ClonerEffector/Public/Cloner/Layouts/CEClonerLayoutBase.h b/Engine/Plugins/Experimental/ClonerEffector/Source/ClonerEffector/Public/Cloner/Layouts/CEClonerLayoutBase.h
--- a/Engine/Plugins/Experimental/ClonerEffector/Source/ClonerEffector/Public/Cloner/Layouts/CEClonerLayoutBase.h
+++ b/Engine/Plugins/Experimental/ClonerEffector/Source/ClonerEffector/Public/Cloner/Layouts/CEClonerLayoutBase.h
@@ -67,6 +67,7 @@ public:
 	UFUNCTION(BlueprintPure, Category="Cloner|Layout")
 	UCEClonerComponent* GetClonerComponent() const;
 
+	/** Get the actor using this layout */
 	AActor* GetClonerActor() const;
 
 	/** Updates all parameters handled by this layout */
@@ -141,6 +142,13 @@ private:
 
 	void OnSystemPackageLoaded(const FName& InName, UPackage* InPackage, EAsyncLoadingResult::Type InResult);
 
+	/** Bind delegates to clear resources during level or world cleanup */
+	void BindCleanupDelegates();
+	void UnbindCleanupDelegates() const;
+
+	void OnWorldCleanup(UWorld* InWorld, bool bInSessionEnded, bool bInCleanupResources);
+	void OnLevelCleanup();
+
 	/** Layout name to display in layout options */
 	UPROPERTY(Transient)
 	FName LayoutName;
@@ -150,11 +158,11 @@ private:
 	FString LayoutAssetPath;
 
 	/** Niagara system representing this layout */
-	UPROPERTY(Transient)
+	UPROPERTY(Transient, DuplicateTransient)
 	TObjectPtr<UNiagaraSystem> NiagaraSystem;
 
 	/** Mesh renderer in this niagara system */
-	UPROPERTY(Transient)
+	UPROPERTY(Transient, DuplicateTransient)
 	TObjectPtr<UNiagaraMeshRendererProperties> MeshRenderer;
 
 	int32 LoadRequestIdentifier = INDEX_NONE;
