MotionDesign : ClonerEffector
- Fixed cloner crash due to cast checked on effector data interface

diff --git a/Engine/Plugins/Experimental/ClonerEffector/Source/ClonerEffector/Private/Cloner/Extensions/CEClonerEffectorExtension.cpp b/Engine/Plugins/Experimental/ClonerEffector/Source/ClonerEffector/Private/Cloner/Extensions/CEClonerEffectorExtension.cpp
--- a/Engine/Plugins/Experimental/ClonerEffector/Source/ClonerEffector/Private/Cloner/Extensions/CEClonerEffectorExtension.cpp
+++ b/Engine/Plugins/Experimental/ClonerEffector/Source/ClonerEffector/Private/Cloner/Extensions/CEClonerEffectorExtension.cpp
@@ -110,7 +110,7 @@ void UCEClonerEffectorExtension::OnEffectorsChanged()
 	const FNiagaraUserRedirectionParameterStore& ExposedParameters = Component->GetOverrideParameters();
 
 	static const FNiagaraVariable EffectorIndexDIVar(FNiagaraTypeDefinition(UNiagaraDataInterfaceArrayInt32::StaticClass()), TEXT("EffectorIndexArray"));
-	UNiagaraDataInterfaceArrayInt32* EffectorIndexArrayDI = CastChecked<UNiagaraDataInterfaceArrayInt32>(ExposedParameters.GetDataInterface(EffectorIndexDIVar));
+	UNiagaraDataInterfaceArrayInt32* EffectorIndexArrayDI = Cast<UNiagaraDataInterfaceArrayInt32>(ExposedParameters.GetDataInterface(EffectorIndexDIVar));
 
 	if (!EffectorIndexArrayDI)
 	{
diff --git a/Engine/Plugins/Experimental/ClonerEffector/Source/ClonerEffector/Private/Cloner/Layouts/CEClonerLayoutBase.cpp b/Engine/Plugins/Experimental/ClonerEffector/Source/ClonerEffector/Private/Cloner/Layouts/CEClonerLayoutBase.cpp
--- a/Engine/Plugins/Experimental/ClonerEffector/Source/ClonerEffector/Private/Cloner/Layouts/CEClonerLayoutBase.cpp
+++ b/Engine/Plugins/Experimental/ClonerEffector/Source/ClonerEffector/Private/Cloner/Layouts/CEClonerLayoutBase.cpp
@@ -297,6 +297,7 @@ void UCEClonerLayoutBase::OnSystemPackageLoaded(const FName& InName, UPackage* I
 	{
 		InPackage->SetFlags(RF_Transient);
 		NiagaraSystem->RemoveFromRoot();
+		NiagaraSystem->ClearFlags(RF_Standalone);
 
 		for (FNiagaraEmitterHandle& SystemEmitterHandle : NiagaraSystem->GetEmitterHandles())
 		{
