MotionDesign : ClonerEffector
- Adding blueprint support for cloner and effector (adding missing setter/getter and function library)

diff --git a/Engine/Plugins/Experimental/Avalanche/Source/AvalancheEffectorsEditor/Private/Cloner/AvaClonerActorVis.cpp b/Engine/Plugins/Experimental/Avalanche/Source/AvalancheEffectorsEditor/Private/Cloner/AvaClonerActorVis.cpp
--- a/Engine/Plugins/Experimental/Avalanche/Source/AvalancheEffectorsEditor/Private/Cloner/AvaClonerActorVis.cpp
+++ b/Engine/Plugins/Experimental/Avalanche/Source/AvalancheEffectorsEditor/Private/Cloner/AvaClonerActorVis.cpp
@@ -747,9 +747,9 @@ UActorComponent* FAvaClonerActorVisualizer::GetEditedComponent() const
 
 TMap<UObject*, TArray<FProperty*>> FAvaClonerActorVisualizer::GatherEditableProperties(UObject* InObject) const
 {
-	if (UCEClonerComponent* ClonerComponent = Cast<UCEClonerComponent>(InObject))
+	if (const UCEClonerComponent* ClonerComponent = Cast<UCEClonerComponent>(InObject))
 	{
-		if (UCEClonerLayoutBase* Layout = ClonerComponent->GetClonerActiveLayout())
+		if (UCEClonerLayoutBase* Layout = ClonerComponent->GetActiveLayout())
 		{
 			TArray<FName> PropertyNames;
 
diff --git a/Engine/Plugins/Experimental/ClonerEffector/Source/ClonerEffector/Private/Cloner/CEClonerComponent.cpp b/Engine/Plugins/Experimental/ClonerEffector/Source/ClonerEffector/Private/Cloner/CEClonerComponent.cpp
--- a/Engine/Plugins/Experimental/ClonerEffector/Source/ClonerEffector/Private/Cloner/CEClonerComponent.cpp
+++ b/Engine/Plugins/Experimental/ClonerEffector/Source/ClonerEffector/Private/Cloner/CEClonerComponent.cpp
@@ -1126,6 +1126,29 @@ void UCEClonerComponent::SetLayoutName(FName InLayoutName)
 	OnLayoutNameChanged();
 }
 
+void UCEClonerComponent::SetLayoutClass(TSubclassOf<UCEClonerLayoutBase> InLayoutClass)
+{
+	if (!InLayoutClass.Get())
+	{
+		return;
+	}
+
+	if (const UCEClonerSubsystem* ClonerSubsystem = UCEClonerSubsystem::Get())
+	{
+		const FName NewLayoutName = ClonerSubsystem->FindLayoutName(InLayoutClass);
+
+		if (!NewLayoutName.IsNone())
+		{
+			SetLayoutName(NewLayoutName);
+		}
+	}
+}
+
+TSubclassOf<UCEClonerLayoutBase> UCEClonerComponent::GetLayoutClass() const
+{
+	return ActiveLayout ? ActiveLayout->GetClass() : nullptr;
+}
+
 #if WITH_EDITOR
 void UCEClonerComponent::SetVisualizerSpriteVisible(bool bInVisible)
 {
@@ -1646,7 +1669,7 @@ TArray<FName> UCEClonerComponent::GetClonerLayoutNames() const
 
 	if (const UCEClonerSubsystem* Subsystem = UCEClonerSubsystem::Get())
 	{
-		LayoutNames = Subsystem->GetLayoutNames();
+		LayoutNames = Subsystem->GetLayoutNames().Array();
 	}
 
 	return LayoutNames;
diff --git a/Engine/Plugins/Experimental/ClonerEffector/Source/ClonerEffector/Private/Cloner/Extensions/CEClonerLifetimeExtension.cpp b/Engine/Plugins/Experimental/ClonerEffector/Source/ClonerEffector/Private/Cloner/Extensions/CEClonerLifetimeExtension.cpp
--- a/Engine/Plugins/Experimental/ClonerEffector/Source/ClonerEffector/Private/Cloner/Extensions/CEClonerLifetimeExtension.cpp
+++ b/Engine/Plugins/Experimental/ClonerEffector/Source/ClonerEffector/Private/Cloner/Extensions/CEClonerLifetimeExtension.cpp
@@ -4,6 +4,7 @@
 
 #include "Cloner/CEClonerComponent.h"
 #include "Cloner/Layouts/CEClonerLayoutBase.h"
+#include "Curves/CurveFloat.h"
 #include "NiagaraDataInterfaceCurve.h"
 #include "NiagaraSystem.h"
 #include "NiagaraUserRedirectionParameterStore.h"
@@ -72,6 +73,16 @@ void UCEClonerLifetimeExtension::SetLifetimeScaleEnabled(bool bInEnabled)
 	MarkExtensionDirty();
 }
 
+void UCEClonerLifetimeExtension::SetLifetimeScaleCurve(UCurveFloat* InCurve)
+{
+	if (!IsValid(InCurve))
+	{
+		return;
+	}
+
+	SetLifetimeScaleCurve(InCurve->FloatCurve);
+}
+
 void UCEClonerLifetimeExtension::SetLifetimeScaleCurve(const FRichCurve& InCurve)
 {
 	LifetimeScaleCurve = InCurve;
diff --git a/Engine/Plugins/Experimental/ClonerEffector/Source/ClonerEffector/Private/Cloner/Extensions/CEClonerMeshRendererExtension.cpp b/Engine/Plugins/Experimental/ClonerEffector/Source/ClonerEffector/Private/Cloner/Extensions/CEClonerMeshRendererExtension.cpp
--- a/Engine/Plugins/Experimental/ClonerEffector/Source/ClonerEffector/Private/Cloner/Extensions/CEClonerMeshRendererExtension.cpp
+++ b/Engine/Plugins/Experimental/ClonerEffector/Source/ClonerEffector/Private/Cloner/Extensions/CEClonerMeshRendererExtension.cpp
@@ -60,25 +60,25 @@ void UCEClonerMeshRendererExtension::SetDefaultMeshes(const TArray<TObjectPtr<US
 	OnOverrideMaterialOptionsChanged();
 }
 
-void UCEClonerMeshRendererExtension::BP_SetDefaultMeshes(const TArray<UStaticMesh*>& InMeshes)
+void UCEClonerMeshRendererExtension::SetDefaultMeshes(const TArray<UStaticMesh*>& InMeshes)
 {
 	DefaultMeshes.Empty(InMeshes.Num());
+
 	Algo::Transform(InMeshes, DefaultMeshes, [](UStaticMesh* InMesh)->TObjectPtr<UStaticMesh>
 	{
 		return InMesh;
 	});
+
 	OnOverrideMaterialOptionsChanged();
 }
 
-TArray<UStaticMesh*> UCEClonerMeshRendererExtension::BP_GetDefaultMeshes() const
+void UCEClonerMeshRendererExtension::GetDefaultMeshes(TArray<UStaticMesh*>& OutMeshes) const
 {
-	TArray<UStaticMesh*> Meshes;
-	Meshes.Reserve(DefaultMeshes.Num());
-	Algo::Transform(DefaultMeshes, Meshes, [](const TObjectPtr<UStaticMesh>& InMesh)->UStaticMesh*
+	OutMeshes.Empty(DefaultMeshes.Num());
+	Algo::Transform(DefaultMeshes, OutMeshes, [](const TObjectPtr<UStaticMesh>& InMesh)->UStaticMesh*
 	{
 		return InMesh;
 	});
-	return Meshes;
 }
 
 void UCEClonerMeshRendererExtension::SetVisualizeEffectors(bool bInVisualize)
diff --git a/Engine/Plugins/Experimental/ClonerEffector/Source/ClonerEffector/Private/Effector/CEEffectorComponent.cpp b/Engine/Plugins/Experimental/ClonerEffector/Source/ClonerEffector/Private/Effector/CEEffectorComponent.cpp
--- a/Engine/Plugins/Experimental/ClonerEffector/Source/ClonerEffector/Private/Effector/CEEffectorComponent.cpp
+++ b/Engine/Plugins/Experimental/ClonerEffector/Source/ClonerEffector/Private/Effector/CEEffectorComponent.cpp
@@ -102,6 +102,29 @@ void UCEEffectorComponent::SetTypeName(FName InTypeName)
 	OnTypeNameChanged();
 }
 
+void UCEEffectorComponent::SetTypeClass(TSubclassOf<UCEEffectorTypeBase> InTypeClass)
+{
+	if (!InTypeClass.Get())
+	{
+		return;
+	}
+
+	if (const UCEEffectorSubsystem* EffectorSubsystem = UCEEffectorSubsystem::Get())
+	{
+		const FName ExtensionName = EffectorSubsystem->FindExtensionName(InTypeClass);
+
+		if (!ExtensionName.IsNone())
+		{
+			SetTypeName(ExtensionName);
+		}
+	}
+}
+
+TSubclassOf<UCEEffectorTypeBase> UCEEffectorComponent::GetTypeClass() const
+{
+	return ActiveType ? ActiveType->GetClass() : nullptr;
+}
+
 void UCEEffectorComponent::SetModeName(FName InModeName)
 {
 	if (ModeName.IsEqual(InModeName))
@@ -118,6 +141,29 @@ void UCEEffectorComponent::SetModeName(FName InModeName)
 	OnModeNameChanged();
 }
 
+void UCEEffectorComponent::SetModeClass(TSubclassOf<UCEEffectorModeBase> InModeClass)
+{
+	if (!InModeClass.Get())
+	{
+		return;
+	}
+
+	if (const UCEEffectorSubsystem* EffectorSubsystem = UCEEffectorSubsystem::Get())
+	{
+		const FName ExtensionName = EffectorSubsystem->FindExtensionName(InModeClass);
+
+		if (!ExtensionName.IsNone())
+		{
+			SetModeName(ExtensionName);
+		}
+	}
+}
+
+TSubclassOf<UCEEffectorModeBase> UCEEffectorComponent::GetModeClass() const
+{
+	return ActiveMode ? ActiveMode->GetClass() : nullptr;
+}
+
 TConstArrayView<TWeakObjectPtr<UCEClonerEffectorExtension>> UCEEffectorComponent::GetClonerExtensionsWeak() const
 {
 	return ClonerExtensionsWeak;
@@ -171,6 +217,43 @@ void UCEEffectorComponent::OnClonerUnlinked(UCEClonerEffectorExtension* InCloner
 	ClonerExtensionsWeak.Remove(InCloner);
 }
 
+void UCEEffectorComponent::GetActiveExtensions(TArray<UCEEffectorExtensionBase*>& OutExtensions) const
+{
+	OutExtensions = ActiveExtensions;
+}
+
+UCEEffectorExtensionBase* UCEEffectorComponent::GetExtension(TSubclassOf<UCEEffectorExtensionBase> InExtensionClass) const
+{
+	const UCEEffectorSubsystem* Subsystem = UCEEffectorSubsystem::Get();
+
+	if (!Subsystem)
+	{
+		return nullptr;
+	}
+
+	const FName ExtensionName = Subsystem->FindExtensionName(InExtensionClass.Get());
+
+	if (ExtensionName.IsNone())
+	{
+		return nullptr;
+	}
+
+	return GetExtension(ExtensionName);
+}
+
+UCEEffectorExtensionBase* UCEEffectorComponent::GetExtension(FName InExtensionName) const
+{
+	for (const TObjectPtr<UCEEffectorExtensionBase>& ExtensionInstance : ExtensionInstances)
+	{
+		if (ExtensionInstance && ExtensionInstance->GetExtensionName() == InExtensionName)
+		{
+			return ExtensionInstance;
+		}
+	}
+
+	return nullptr;
+}
+
 void UCEEffectorComponent::RequestClonerUpdate(bool bInImmediate)
 {
 	if (bInImmediate)
@@ -669,7 +752,7 @@ TArray<FName> UCEEffectorComponent::GetEffectorTypeNames() const
 
 	if (const UCEEffectorSubsystem* EffectorSubsystem = UCEEffectorSubsystem::Get())
 	{
-		TypeNames = EffectorSubsystem->GetExtensionNames<UCEEffectorTypeBase>();
+		TypeNames = EffectorSubsystem->GetExtensionNames<UCEEffectorTypeBase>().Array();
 	}
 
 	return TypeNames;
@@ -681,7 +764,7 @@ TArray<FName> UCEEffectorComponent::GetEffectorModeNames() const
 
 	if (const UCEEffectorSubsystem* EffectorSubsystem = UCEEffectorSubsystem::Get())
 	{
-		ModeNames = EffectorSubsystem->GetExtensionNames<UCEEffectorModeBase>();
+		ModeNames = EffectorSubsystem->GetExtensionNames<UCEEffectorModeBase>().Array();
 	}
 
 	return ModeNames;
diff --git a/Engine/Plugins/Experimental/ClonerEffector/Source/ClonerEffector/Private/Subsystems/CEClonerSubsystem.cpp b/Engine/Plugins/Experimental/ClonerEffector/Source/ClonerEffector/Private/Subsystems/CEClonerSubsystem.cpp
--- a/Engine/Plugins/Experimental/ClonerEffector/Source/ClonerEffector/Private/Subsystems/CEClonerSubsystem.cpp
+++ b/Engine/Plugins/Experimental/ClonerEffector/Source/ClonerEffector/Private/Subsystems/CEClonerSubsystem.cpp
@@ -437,11 +437,18 @@ AActor* UCEClonerSubsystem::CreateClonerWithActors(UWorld* InWorld, const TSet<A
 	return NewClonerActor;
 }
 
-TArray<FName> UCEClonerSubsystem::GetLayoutNames() const
+TSet<FName> UCEClonerSubsystem::GetLayoutNames() const
 {
 	TArray<FName> LayoutNames;
 	LayoutClasses.GenerateKeyArray(LayoutNames);
-	return LayoutNames;
+	return TSet<FName>(LayoutNames);
+}
+
+TSet<TSubclassOf<UCEClonerLayoutBase>> UCEClonerSubsystem::GetLayoutClasses() const
+{
+	TArray<TSubclassOf<UCEClonerLayoutBase>> Layouts;
+	LayoutClasses.GenerateValueArray(Layouts);
+	return TSet<TSubclassOf<UCEClonerLayoutBase>>(Layouts);
 }
 
 FName UCEClonerSubsystem::FindLayoutName(TSubclassOf<UCEClonerLayoutBase> InLayoutClass) const
diff --git a/Engine/Plugins/Experimental/ClonerEffector/Source/ClonerEffector/Private/Subsystems/CEEffectorSubsystem.cpp b/Engine/Plugins/Experimental/ClonerEffector/Source/ClonerEffector/Private/Subsystems/CEEffectorSubsystem.cpp
--- a/Engine/Plugins/Experimental/ClonerEffector/Source/ClonerEffector/Private/Subsystems/CEEffectorSubsystem.cpp
+++ b/Engine/Plugins/Experimental/ClonerEffector/Source/ClonerEffector/Private/Subsystems/CEEffectorSubsystem.cpp
@@ -192,9 +192,9 @@ bool UCEEffectorSubsystem::IsExtensionClassRegistered(UClass* InClass) const
 	return !!ExtensionClasses.FindKey(ExtensionClass);
 }
 
-TArray<FName> UCEEffectorSubsystem::GetExtensionNames(TSubclassOf<UCEEffectorExtensionBase> InExtensionClass) const
+TSet<FName> UCEEffectorSubsystem::GetExtensionNames(TSubclassOf<UCEEffectorExtensionBase> InExtensionClass) const
 {
-	TArray<FName> ExtensionNames;
+	TSet<FName> ExtensionNames;
 
 	const FName ExtensionName = FindExtensionName(InExtensionClass);
 
@@ -216,6 +216,30 @@ TArray<FName> UCEEffectorSubsystem::GetExtensionNames(TSubclassOf<UCEEffectorExt
 	return ExtensionNames;
 }
 
+TSet<TSubclassOf<UCEEffectorExtensionBase>> UCEEffectorSubsystem::GetExtensionClasses(TSubclassOf<UCEEffectorExtensionBase> InExtensionClass) const
+{
+	TSet<TSubclassOf<UCEEffectorExtensionBase>> Extensions;
+
+	if (!InExtensionClass.Get())
+	{
+		return Extensions;
+	}
+
+	Extensions.Empty(ExtensionClasses.Num());
+
+	for (const TPair<FName, TSubclassOf<UCEEffectorExtensionBase>>& ExtensionPair : ExtensionClasses)
+	{
+		if (ExtensionPair.Value
+			&& (ExtensionPair.Value == InExtensionClass
+				|| ExtensionPair.Value->IsChildOf(InExtensionClass)))
+		{
+			Extensions.Add(ExtensionPair.Value);
+		}
+	}
+
+	return Extensions;
+}
+
 FName UCEEffectorSubsystem::FindExtensionName(TSubclassOf<UCEEffectorExtensionBase> InClass) const
 {
 	if (const FName* ExtensionName = ExtensionClasses.FindKey(InClass))
diff --git a/Engine/Plugins/Experimental/ClonerEffector/Source/ClonerEffector/Private/Utilities/CEClonerLibrary.cpp b/Engine/Plugins/Experimental/ClonerEffector/Source/ClonerEffector/Private/Utilities/CEClonerLibrary.cpp
new file mode 100644
--- /dev/null
+++ b/Engine/Plugins/Experimental/ClonerEffector/Source/ClonerEffector/Private/Utilities/CEClonerLibrary.cpp
@@ -0,0 +1,26 @@
+// Copyright Epic Games, Inc. All Rights Reserved.
+
+#include "Utilities/CEClonerLibrary.h"
+
+#include "Containers/Set.h"
+#include "Subsystems/CEClonerSubsystem.h"
+
+void UCEClonerLibrary::GetClonerLayoutClasses(TSet<TSubclassOf<UCEClonerLayoutBase>>& OutLayoutClasses)
+{
+	OutLayoutClasses.Empty();
+
+	if (const UCEClonerSubsystem* ClonerSubsystem = UCEClonerSubsystem::Get())
+	{
+		OutLayoutClasses = ClonerSubsystem->GetLayoutClasses();
+	}
+}
+
+void UCEClonerLibrary::GetClonerExtensionClasses(TSet<TSubclassOf<UCEClonerExtensionBase>>& OutExtensionClasses)
+{
+	OutExtensionClasses.Empty();
+
+	if (const UCEClonerSubsystem* ClonerSubsystem = UCEClonerSubsystem::Get())
+	{
+		OutExtensionClasses = ClonerSubsystem->GetExtensionClasses();
+	}
+}
diff --git a/Engine/Plugins/Experimental/ClonerEffector/Source/ClonerEffector/Private/Utilities/CEEffectorLibrary.cpp b/Engine/Plugins/Experimental/ClonerEffector/Source/ClonerEffector/Private/Utilities/CEEffectorLibrary.cpp
new file mode 100644
--- /dev/null
+++ b/Engine/Plugins/Experimental/ClonerEffector/Source/ClonerEffector/Private/Utilities/CEEffectorLibrary.cpp
@@ -0,0 +1,54 @@
+// Copyright Epic Games, Inc. All Rights Reserved.
+
+#include "Utilities/CEEffectorLibrary.h"
+
+#include "Containers/Set.h"
+#include "Effector/Modes/CEEffectorModeBase.h"
+#include "Effector/Types/CEEffectorTypeBase.h"
+#include "Subsystems/CEEffectorSubsystem.h"
+
+void UCEEffectorLibrary::GetEffectorModeClasses(TSet<TSubclassOf<UCEEffectorModeBase>>& OutModeClasses)
+{
+	OutModeClasses.Empty();
+
+	if (const UCEEffectorSubsystem* EffectorSubsystem = UCEEffectorSubsystem::Get())
+	{
+		for (TSubclassOf<UCEEffectorExtensionBase> ModeClass : EffectorSubsystem->GetExtensionClasses<UCEEffectorModeBase>())
+		{
+			OutModeClasses.Add(ModeClass.Get());
+		}
+	}
+}
+
+void UCEEffectorLibrary::GetEffectorTypeClasses(TSet<TSubclassOf<UCEEffectorTypeBase>>& OutTypeClasses)
+{
+	OutTypeClasses.Empty();
+
+	if (const UCEEffectorSubsystem* EffectorSubsystem = UCEEffectorSubsystem::Get())
+	{
+		for (TSubclassOf<UCEEffectorExtensionBase> TypeClass : EffectorSubsystem->GetExtensionClasses<UCEEffectorTypeBase>())
+		{
+			OutTypeClasses.Add(TypeClass.Get());
+		}
+	}
+}
+
+void UCEEffectorLibrary::GetEffectorExtensionClasses(TSet<TSubclassOf<UCEEffectorExtensionBase>>& OutExtensionClasses)
+{
+	OutExtensionClasses.Empty();
+
+	if (const UCEEffectorSubsystem* EffectorSubsystem = UCEEffectorSubsystem::Get())
+	{
+		for (TSubclassOf<UCEEffectorExtensionBase> ExtensionClass : EffectorSubsystem->GetExtensionClasses<UCEEffectorExtensionBase>())
+		{
+			if (!ExtensionClass.Get()
+				|| ExtensionClass->IsChildOf<UCEEffectorTypeBase>()
+				|| ExtensionClass->IsChildOf<UCEEffectorModeBase>())
+			{
+				continue;
+			}
+
+			OutExtensionClasses.Add(ExtensionClass);
+		}
+	}
+}
diff --git a/Engine/Plugins/Experimental/ClonerEffector/Source/ClonerEffector/Public/Cloner/CEClonerActor.h b/Engine/Plugins/Experimental/ClonerEffector/Source/ClonerEffector/Public/Cloner/CEClonerActor.h
--- a/Engine/Plugins/Experimental/ClonerEffector/Source/ClonerEffector/Public/Cloner/CEClonerActor.h
+++ b/Engine/Plugins/Experimental/ClonerEffector/Source/ClonerEffector/Public/Cloner/CEClonerActor.h
@@ -30,6 +30,7 @@ public:
 
 	ACEClonerActor();
 
+	UFUNCTION(BlueprintPure, Category="Cloner")
 	UCEClonerComponent* GetClonerComponent() const
 	{
 		return ClonerComponent;
diff --git a/Engine/Plugins/Experimental/ClonerEffector/Source/ClonerEffector/Public/Cloner/CEClonerComponent.h b/Engine/Plugins/Experimental/ClonerEffector/Source/ClonerEffector/Public/Cloner/CEClonerComponent.h
--- a/Engine/Plugins/Experimental/ClonerEffector/Source/ClonerEffector/Public/Cloner/CEClonerComponent.h
+++ b/Engine/Plugins/Experimental/ClonerEffector/Source/ClonerEffector/Public/Cloner/CEClonerComponent.h
@@ -26,6 +26,10 @@ class UCEClonerComponent : public UNiagaraComponent
 
 	friend class ACEClonerActor;
 
+	DECLARE_MULTICAST_DELEGATE_OneParam(FOnClonerMeshUpdated, UCEClonerComponent* /** ClonerComponent */)
+	DECLARE_MULTICAST_DELEGATE_TwoParams(FOnClonerLayoutLoaded, UCEClonerComponent* /** ClonerComponent */, UCEClonerLayoutBase* /** InLayout */)
+	DECLARE_MULTICAST_DELEGATE_OneParam(FOnClonerInitialized, UCEClonerComponent* /** ClonerComponent */)
+
 public:
 	/** Only materials transient or part of the content folder can be dirtied, engine or plugins cannot */
 	static bool IsMaterialDirtyable(const UMaterialInterface* InMaterial);
@@ -44,17 +48,17 @@ public:
 	static CLONEREFFECTOR_API FName GetLayoutNamePropertyName();
 #endif
 
-	static TMulticastDelegateRegistration<void(UCEClonerComponent*)>& OnClonerMeshUpdated()
+	static FOnClonerMeshUpdated::RegistrationType& OnClonerMeshUpdated()
 	{
 		return OnClonerMeshUpdatedDelegate;
 	}
 
-	static TMulticastDelegateRegistration<void(UCEClonerComponent*, UCEClonerLayoutBase*)>& OnClonerLayoutLoaded()
+	static FOnClonerLayoutLoaded::RegistrationType& OnClonerLayoutLoaded()
 	{
 		return OnClonerLayoutLoadedDelegate;
 	}
 
-	static TMulticastDelegateRegistration<void(UCEClonerComponent*)>& OnClonerInitialized()
+	static FOnClonerInitialized::RegistrationType& OnClonerInitialized()
 	{
 		return OnClonerInitializedDelegate;
 	}
@@ -106,6 +110,12 @@ public:
 		return LayoutName;
 	}
 
+	UFUNCTION(BlueprintCallable, Category="Cloner")
+	void SetLayoutClass(TSubclassOf<UCEClonerLayoutBase> InLayoutClass);
+
+	UFUNCTION(BlueprintPure, Category="Cloner")
+	TSubclassOf<UCEClonerLayoutBase> GetLayoutClass() const;
+
 	UFUNCTION(BlueprintPure, Category="Cloner")
 	UCEClonerLayoutBase* GetActiveLayout() const
 	{
@@ -185,11 +195,6 @@ public:
 	/** Will force a system update to refresh user parameters */
 	void RequestClonerUpdate(bool bInImmediate = false);
 
-	UCEClonerLayoutBase* GetClonerActiveLayout() const
-	{
-		return ActiveLayout;
-	}
-
 	/** Forces a refresh of the meshes used */
 	void RefreshClonerMeshes();
 
@@ -200,6 +205,7 @@ public:
 		return Cast<InExtensionClass>(GetExtension(InExtensionClass::StaticClass()));
 	}
 
+	UFUNCTION(BlueprintCallable, Category="Cloner")
 	UCEClonerExtensionBase* GetExtension(TSubclassOf<UCEClonerExtensionBase> InExtensionClass) const;
 
 	UCEClonerExtensionBase* GetExtension(FName InExtensionName) const;
@@ -209,17 +215,24 @@ public:
 		return ActiveExtensions;
 	}
 
+	/**
+	 * Retrieves all active extensions on this cloner
+	 * @param OutExtensions [Out] Active extensions
+	 */
+	UFUNCTION(BlueprintCallable, Category="Cloner")
+	void GetActiveExtensions(TArray<UCEClonerExtensionBase*>& OutExtensions) const
+	{
+		OutExtensions = ActiveExtensions;
+	}
+
 protected:
 	/** Called when meshes have been updated */
-	DECLARE_MULTICAST_DELEGATE_OneParam(FOnClonerMeshUpdated, UCEClonerComponent* /** ClonerComponent */)
 	CLONEREFFECTOR_API static FOnClonerMeshUpdated OnClonerMeshUpdatedDelegate;
 
 	/** Called when new cloner layout is loaded */
-	DECLARE_MULTICAST_DELEGATE_TwoParams(FOnClonerLayoutLoaded, UCEClonerComponent* /** ClonerComponent */, UCEClonerLayoutBase* /** InLayout */)
 	CLONEREFFECTOR_API static FOnClonerLayoutLoaded OnClonerLayoutLoadedDelegate;
 
 	/** Called when cloner is initialized */
-	DECLARE_MULTICAST_DELEGATE_OneParam(FOnClonerInitialized, UCEClonerComponent* /** ClonerComponent */)
 	CLONEREFFECTOR_API static FOnClonerInitialized OnClonerInitializedDelegate;
 
 	//~ Begin UObject
@@ -276,23 +289,23 @@ protected:
 	void OnActiveLayoutChanged();
 
 	/** Is this cloner enabled/disabled */
-	UPROPERTY(EditInstanceOnly, BlueprintReadWrite, Setter="SetEnabled", Getter="GetEnabled", Category="Cloner")
+	UPROPERTY(EditInstanceOnly, Setter="SetEnabled", Getter="GetEnabled", Category="Cloner")
 	bool bEnabled = true;
 
 	/** Interval to update the attachment tree and update the cloner meshes, 0 means each tick */
-	UPROPERTY(EditInstanceOnly, BlueprintReadWrite, Setter, Getter, Category="Cloner", AdvancedDisplay, meta=(ClampMin="0"))
+	UPROPERTY(EditInstanceOnly, Setter, Getter, Category="Cloner", AdvancedDisplay, meta=(ClampMin="0"))
 	float TreeUpdateInterval = 0.2f;
 
 	/** Cloner instance seed for random deterministic patterns */
-	UPROPERTY(EditInstanceOnly, BlueprintReadWrite, Setter, Getter, Category="Cloner")
+	UPROPERTY(EditInstanceOnly, Setter, Getter, Category="Cloner")
 	int32 Seed = 0;
 
 	/** Cloner color when unaffected by effectors, color will be passed down to the material (ParticleColor) */
-	UPROPERTY(EditInstanceOnly, BlueprintReadWrite, Setter, Getter, Category="Cloner")
+	UPROPERTY(EditInstanceOnly, Setter, Getter, Category="Cloner")
 	FLinearColor Color = FLinearColor::White;
 
 	/** Name of the layout to use */
-	UPROPERTY(EditInstanceOnly, BlueprintReadWrite, Setter, Getter, Category="Layout", meta=(GetOptions="GetClonerLayoutNames"))
+	UPROPERTY(EditInstanceOnly, Setter, Getter, Category="Layout", meta=(GetOptions="GetClonerLayoutNames"))
 	FName LayoutName = NAME_None;
 
 	/** Active layout used */
@@ -313,7 +326,7 @@ protected:
 
 #if WITH_EDITORONLY_DATA
 	/** Toggle the sprite to visualize and click on this cloner */
-	UPROPERTY(EditInstanceOnly, BlueprintReadWrite, AdvancedDisplay, Category="Cloner", meta=(AllowPrivateAccess="true"))
+	UPROPERTY(EditInstanceOnly, AdvancedDisplay, Category="Cloner", meta=(AllowPrivateAccess="true"))
 	bool bVisualizerSpriteVisible = true;
 #endif
 
diff --git a/Engine/Plugins/Experimental/ClonerEffector/Source/ClonerEffector/Public/Cloner/Extensions/CEClonerCollisionExtension.h b/Engine/Plugins/Experimental/ClonerEffector/Source/ClonerEffector/Public/Cloner/Extensions/CEClonerCollisionExtension.h
--- a/Engine/Plugins/Experimental/ClonerEffector/Source/ClonerEffector/Public/Cloner/Extensions/CEClonerCollisionExtension.h
+++ b/Engine/Plugins/Experimental/ClonerEffector/Source/ClonerEffector/Public/Cloner/Extensions/CEClonerCollisionExtension.h
@@ -124,43 +124,43 @@ protected:
 	 * Allow particles to react to surface by using the distance field,
 	 * ensure the mesh you want particle to collide with is tick enough or increase its distance field resolution scale in static mesh editor 
 	 */
-	UPROPERTY(EditInstanceOnly, BlueprintReadWrite, Setter="SetSurfaceCollisionEnabled", Getter="GetSurfaceCollisionEnabled", Category="Collisions")
+	UPROPERTY(EditInstanceOnly, Setter="SetSurfaceCollisionEnabled", Getter="GetSurfaceCollisionEnabled", Category="Collisions")
 	bool bSurfaceCollisionEnabled = false;
 
 	/** Allow particles to react to other emitter particles, uses a neighbor grid to detect collision */
-	UPROPERTY(EditInstanceOnly, BlueprintReadWrite, Setter="SetParticleCollisionEnabled", Getter="GetParticleCollisionEnabled", Category="Collisions")
+	UPROPERTY(EditInstanceOnly, Setter="SetParticleCollisionEnabled", Getter="GetParticleCollisionEnabled", Category="Collisions")
 	bool bParticleCollisionEnabled = false;
 
 	/** Recalculate accurate velocity after collision is updated */
-	UPROPERTY(EditInstanceOnly, BlueprintReadWrite, Setter="SetCollisionVelocityEnabled", Getter="GetCollisionVelocityEnabled", Category="Collisions", meta=(EditCondition="bParticleCollisionEnabled", EditConditionHides))
+	UPROPERTY(EditInstanceOnly, Setter="SetCollisionVelocityEnabled", Getter="GetCollisionVelocityEnabled", Category="Collisions", meta=(EditCondition="bParticleCollisionEnabled", EditConditionHides))
 	bool bCollisionVelocityEnabled = false;
 
 	/** Amount of iterations to improve particle collision results but affects performance */
-	UPROPERTY(EditInstanceOnly, BlueprintReadWrite, Setter, Getter, Category="Collisions", meta=(ClampMin="1", EditCondition="bParticleCollisionEnabled", EditConditionHides))
+	UPROPERTY(EditInstanceOnly, Setter, Getter, Category="Collisions", meta=(ClampMin="1", EditCondition="bParticleCollisionEnabled", EditConditionHides))
 	int32 CollisionIterations = 1;
 
 	/** Resolution of the neighbor grid to detect collision between particles */
-	UPROPERTY(EditInstanceOnly, BlueprintReadWrite, Setter, Getter, Category="Collisions", meta=(ClampMin="1", EditCondition="bParticleCollisionEnabled", EditConditionHides))
+	UPROPERTY(EditInstanceOnly, Setter, Getter, Category="Collisions", meta=(ClampMin="1", EditCondition="bParticleCollisionEnabled", EditConditionHides))
 	int32 CollisionGridResolution = 32;
 
 	/** Size of the neighbor grid to detect collision between particles, particles outside this grid will not have collisions */
-	UPROPERTY(EditInstanceOnly, BlueprintReadWrite, Setter, Getter, Category="Collisions", meta=(ClampMin="0", MotionDesignVectorWidget, AllowPreserveRatio="XYZ", EditCondition="bParticleCollisionEnabled", EditConditionHides))
+	UPROPERTY(EditInstanceOnly, Setter, Getter, Category="Collisions", meta=(ClampMin="0", MotionDesignVectorWidget, AllowPreserveRatio="XYZ", EditCondition="bParticleCollisionEnabled", EditConditionHides))
 	FVector CollisionGridSize = FVector(5000.f);
 
 	/** Collision radius calculation mode */
-	UPROPERTY(EditInstanceOnly, BlueprintReadWrite, Setter, Getter, Category="Collisions", meta=(EditCondition="bParticleCollisionEnabled", EditConditionHides))
+	UPROPERTY(EditInstanceOnly, Setter, Getter, Category="Collisions", meta=(EditCondition="bParticleCollisionEnabled", EditConditionHides))
 	ECEClonerCollisionRadiusMode CollisionRadiusMode = ECEClonerCollisionRadiusMode::ExtentLength;
 
 	/** Radius expected around each particle for collision, order matches attachment index */
-	UPROPERTY(EditInstanceOnly, BlueprintReadWrite, Getter, Category="Collisions", EditFixedSize, meta=(ClampMin="0", EditFixedOrder, EditCondition="bParticleCollisionEnabled && CollisionRadiusMode == ECEClonerCollisionRadiusMode::Manual", EditConditionHides))
+	UPROPERTY(EditInstanceOnly, Getter, Category="Collisions", EditFixedSize, meta=(ClampMin="0", EditFixedOrder, EditCondition="bParticleCollisionEnabled && CollisionRadiusMode == ECEClonerCollisionRadiusMode::Manual", EditConditionHides))
 	TArray<float> CollisionRadii;
 
 	/** Minimum particle mass, used for collisions to push apart */
-	UPROPERTY(EditInstanceOnly, BlueprintReadWrite, Setter, Getter, Category="Collisions", meta=(ClampMin="1"))
+	UPROPERTY(EditInstanceOnly, Setter, Getter, Category="Collisions", meta=(ClampMin="1"))
 	float MassMin = 1.f;
 
 	/** Maximum particle mass, used for collisions to push apart */
-	UPROPERTY(EditInstanceOnly, BlueprintReadWrite, Setter, Getter, Category="Collisions", meta=(ClampMin="1"))
+	UPROPERTY(EditInstanceOnly, Setter, Getter, Category="Collisions", meta=(ClampMin="1"))
 	float MassMax = 2.f;
 
 private:
diff --git a/Engine/Plugins/Experimental/ClonerEffector/Source/ClonerEffector/Public/Cloner/Extensions/CEClonerConstraintExtension.h b/Engine/Plugins/Experimental/ClonerEffector/Source/ClonerEffector/Public/Cloner/Extensions/CEClonerConstraintExtension.h
--- a/Engine/Plugins/Experimental/ClonerEffector/Source/ClonerEffector/Public/Cloner/Extensions/CEClonerConstraintExtension.h
+++ b/Engine/Plugins/Experimental/ClonerEffector/Source/ClonerEffector/Public/Cloner/Extensions/CEClonerConstraintExtension.h
@@ -140,46 +140,46 @@ protected:
 	virtual bool IsLayoutSupported(const UCEClonerLayoutBase* InLayout) const override;
 	//~ End UCEClonerExtensionBase
 
-	UPROPERTY(EditInstanceOnly, BlueprintReadWrite, Setter, Getter, Category="Constraint")
+	UPROPERTY(EditInstanceOnly, Setter, Getter, Category="Constraint")
 	ECEClonerGridConstraint Constraint = ECEClonerGridConstraint::None;
 
-	UPROPERTY(EditInstanceOnly, BlueprintReadWrite, Setter="SetInvertConstraint", Getter="GetInvertConstraint", Category="Constraint", meta=(EditCondition="Constraint != ECEClonerGridConstraint::None", EditConditionHides))
+	UPROPERTY(EditInstanceOnly, Setter="SetInvertConstraint", Getter="GetInvertConstraint", Category="Constraint", meta=(EditCondition="Constraint != ECEClonerGridConstraint::None", EditConditionHides))
 	bool bInvertConstraint = false;
 
 	/** Sphere constraint */
 
-	UPROPERTY(EditInstanceOnly, BlueprintReadWrite, Setter, Getter, Category="Constraint", meta=(EditCondition="Constraint == ECEClonerGridConstraint::Sphere", EditConditionHides))
+	UPROPERTY(EditInstanceOnly, Setter, Getter, Category="Constraint", meta=(EditCondition="Constraint == ECEClonerGridConstraint::Sphere", EditConditionHides))
 	float SphereRadius = 400.f;
 
-	UPROPERTY(EditInstanceOnly, BlueprintReadWrite, Setter, Getter, Category="Constraint", meta=(EditCondition="Constraint == ECEClonerGridConstraint::Sphere", EditConditionHides))
+	UPROPERTY(EditInstanceOnly, Setter, Getter, Category="Constraint", meta=(EditCondition="Constraint == ECEClonerGridConstraint::Sphere", EditConditionHides))
 	FVector SphereCenter = FVector::ZeroVector;
 
 	/** Cylinder constraint */
 
-	UPROPERTY(EditInstanceOnly, BlueprintReadWrite, Setter, Getter, Category="Constraint", meta=(EditCondition="Constraint == ECEClonerGridConstraint::Cylinder", EditConditionHides))
+	UPROPERTY(EditInstanceOnly, Setter, Getter, Category="Constraint", meta=(EditCondition="Constraint == ECEClonerGridConstraint::Cylinder", EditConditionHides))
 	float CylinderRadius = 400.f;
 
-	UPROPERTY(EditInstanceOnly, BlueprintReadWrite, Setter, Getter, Category="Constraint", meta=(EditCondition="Constraint == ECEClonerGridConstraint::Cylinder", EditConditionHides))
+	UPROPERTY(EditInstanceOnly, Setter, Getter, Category="Constraint", meta=(EditCondition="Constraint == ECEClonerGridConstraint::Cylinder", EditConditionHides))
 	float CylinderHeight = 800.f;
 
-	UPROPERTY(EditInstanceOnly, BlueprintReadWrite, Setter, Getter, Category="Constraint", meta=(EditCondition="Constraint == ECEClonerGridConstraint::Cylinder", EditConditionHides))
+	UPROPERTY(EditInstanceOnly, Setter, Getter, Category="Constraint", meta=(EditCondition="Constraint == ECEClonerGridConstraint::Cylinder", EditConditionHides))
 	FVector CylinderCenter = FVector::ZeroVector;
 
 	/** Texture constraint */
 
-	UPROPERTY(EditInstanceOnly, BlueprintReadWrite, Setter, Getter, Category="Constraint", meta=(EditCondition="Constraint == ECEClonerGridConstraint::Texture", EditConditionHides))
+	UPROPERTY(EditInstanceOnly, Setter, Getter, Category="Constraint", meta=(EditCondition="Constraint == ECEClonerGridConstraint::Texture", EditConditionHides))
 	TObjectPtr<UTexture> TextureAsset;
 
-	UPROPERTY(EditInstanceOnly, BlueprintReadWrite, Setter, Getter, Category="Constraint", meta=(InvalidEnumValues="Custom", EditCondition="Constraint == ECEClonerGridConstraint::Texture", EditConditionHides))
+	UPROPERTY(EditInstanceOnly, Setter, Getter, Category="Constraint", meta=(InvalidEnumValues="Custom", EditCondition="Constraint == ECEClonerGridConstraint::Texture", EditConditionHides))
 	ECEClonerPlane TexturePlane = ECEClonerPlane::XY;
 
-	UPROPERTY(EditInstanceOnly, BlueprintReadWrite, Setter, Getter, Category="Constraint", meta=(EditCondition="Constraint == ECEClonerGridConstraint::Texture", EditConditionHides))
+	UPROPERTY(EditInstanceOnly, Setter, Getter, Category="Constraint", meta=(EditCondition="Constraint == ECEClonerGridConstraint::Texture", EditConditionHides))
 	ECEClonerTextureSampleChannel TextureSampleMode = ECEClonerTextureSampleChannel::RGBLuminance;
 
-	UPROPERTY(EditInstanceOnly, BlueprintReadWrite, Setter, Getter, Category="Constraint", meta=(EditCondition="Constraint == ECEClonerGridConstraint::Texture", EditConditionHides))
+	UPROPERTY(EditInstanceOnly, Setter, Getter, Category="Constraint", meta=(EditCondition="Constraint == ECEClonerGridConstraint::Texture", EditConditionHides))
 	ECEClonerCompareMode TextureCompareMode = ECEClonerCompareMode::Greater;
 
-	UPROPERTY(EditInstanceOnly, BlueprintReadWrite, Setter, Getter, Category="Constraint", meta=(EditCondition="Constraint == ECEClonerGridConstraint::Texture", EditConditionHides))
+	UPROPERTY(EditInstanceOnly, Setter, Getter, Category="Constraint", meta=(EditCondition="Constraint == ECEClonerGridConstraint::Texture", EditConditionHides))
 	float TextureThreshold = 0.f;
 
 private:
diff --git a/Engine/Plugins/Experimental/ClonerEffector/Source/ClonerEffector/Public/Cloner/Extensions/CEClonerEffectorExtension.h b/Engine/Plugins/Experimental/ClonerEffector/Source/ClonerEffector/Public/Cloner/Extensions/CEClonerEffectorExtension.h
--- a/Engine/Plugins/Experimental/ClonerEffector/Source/ClonerEffector/Public/Cloner/Extensions/CEClonerEffectorExtension.h
+++ b/Engine/Plugins/Experimental/ClonerEffector/Source/ClonerEffector/Public/Cloner/Extensions/CEClonerEffectorExtension.h
@@ -34,12 +34,15 @@ public:
 	CLONEREFFECTOR_API int32 GetEffectorCount() const;
 
 	/** Links new actor effectors to apply transformation on clones */
+	UFUNCTION(BlueprintCallable, Category="Cloner")
 	CLONEREFFECTOR_API bool LinkEffector(AActor* InEffectorActor);
 
 	/** Unlinks the effector actor and reset the cloner simulation */
+	UFUNCTION(BlueprintCallable, Category="Cloner")
 	CLONEREFFECTOR_API bool UnlinkEffector(AActor* InEffectorActor);
 
 	/** Checks if an effector is linked with this cloner */
+	UFUNCTION(BlueprintPure, Category="Cloner")
 	CLONEREFFECTOR_API bool IsEffectorLinked(const AActor* InEffectorActor) const;
 
 #if WITH_EDITOR
diff --git a/Engine/Plugins/Experimental/ClonerEffector/Source/ClonerEffector/Public/Cloner/Extensions/CEClonerEmitterSpawnExtension.h b/Engine/Plugins/Experimental/ClonerEffector/Source/ClonerEffector/Public/Cloner/Extensions/CEClonerEmitterSpawnExtension.h
--- a/Engine/Plugins/Experimental/ClonerEffector/Source/ClonerEffector/Public/Cloner/Extensions/CEClonerEmitterSpawnExtension.h
+++ b/Engine/Plugins/Experimental/ClonerEffector/Source/ClonerEffector/Public/Cloner/Extensions/CEClonerEmitterSpawnExtension.h
@@ -80,23 +80,23 @@ protected:
 	//~ End UCEClonerExtensionBase
 
 	/** How many times do we spawn clones */
-	UPROPERTY(EditInstanceOnly, BlueprintReadWrite, Setter="SetSpawnLoopMode", Getter="GetSpawnLoopMode", DisplayName="Emission Mode", Category="Spawn")
+	UPROPERTY(EditInstanceOnly, Setter, Getter, DisplayName="Emission Mode", Category="Spawn")
 	ECEClonerSpawnLoopMode SpawnLoopMode = ECEClonerSpawnLoopMode::Once;
 
 	/** Amount of spawn iterations for clones */
-	UPROPERTY(EditInstanceOnly, BlueprintReadWrite, Setter="SetSpawnLoopIterations", Getter="GetSpawnLoopIterations", DisplayName="Emission Count", Category="Spawn", meta=(ClampMin="1", EditCondition="SpawnLoopMode == ECEClonerSpawnLoopMode::Multiple", EditConditionHides))
+	UPROPERTY(EditInstanceOnly, Setter, Getter, DisplayName="Emission Count", Category="Spawn", meta=(ClampMin="1", EditCondition="SpawnLoopMode == ECEClonerSpawnLoopMode::Multiple", EditConditionHides))
 	int32 SpawnLoopIterations = 1;
 
 	/** Interval/Duration of spawn for clones */
-	UPROPERTY(EditInstanceOnly, BlueprintReadWrite, Setter="SetSpawnLoopInterval", Getter="GetSpawnLoopInterval", DisplayName="Emission Interval", Category="Spawn", meta=(ClampMin="0", EditCondition="SpawnLoopMode != ECEClonerSpawnLoopMode::Once", EditConditionHides))
+	UPROPERTY(EditInstanceOnly, Setter, Getter, DisplayName="Emission Interval", Category="Spawn", meta=(ClampMin="0", EditCondition="SpawnLoopMode != ECEClonerSpawnLoopMode::Once", EditConditionHides))
 	float SpawnLoopInterval = 1.f;
 
 	/** How does spawn occurs */
-	UPROPERTY(EditInstanceOnly, BlueprintReadWrite, Setter="SetSpawnBehaviorMode", Getter="GetSpawnBehaviorMode", DisplayName="Emission Style", Category="Spawn", meta=(EditCondition="SpawnLoopMode != ECEClonerSpawnLoopMode::Once", EditConditionHides))
+	UPROPERTY(EditInstanceOnly, Setter, Getter, DisplayName="Emission Style", Category="Spawn", meta=(EditCondition="SpawnLoopMode != ECEClonerSpawnLoopMode::Once", EditConditionHides))
 	ECEClonerSpawnBehaviorMode SpawnBehaviorMode = ECEClonerSpawnBehaviorMode::Instant;
 
 	/** How many clones to spawn each seconds */
-	UPROPERTY(EditInstanceOnly, BlueprintReadWrite, Setter="SetSpawnRate", Getter="GetSpawnRate", DisplayName="Emission Rate", Category="Spawn", meta=(ClampMin="0", EditCondition="SpawnLoopMode != ECEClonerSpawnLoopMode::Once && SpawnBehaviorMode == ECEClonerSpawnBehaviorMode::Rate", EditConditionHides))
+	UPROPERTY(EditInstanceOnly, Setter, Getter, DisplayName="Emission Rate", Category="Spawn", meta=(ClampMin="0", EditCondition="SpawnLoopMode != ECEClonerSpawnLoopMode::Once && SpawnBehaviorMode == ECEClonerSpawnBehaviorMode::Rate", EditConditionHides))
 	float SpawnRate = 1.f;
 
 private:
diff --git a/Engine/Plugins/Experimental/ClonerEffector/Source/ClonerEffector/Public/Cloner/Extensions/CEClonerExtensionBase.h b/Engine/Plugins/Experimental/ClonerEffector/Source/ClonerEffector/Public/Cloner/Extensions/CEClonerExtensionBase.h
--- a/Engine/Plugins/Experimental/ClonerEffector/Source/ClonerEffector/Public/Cloner/Extensions/CEClonerExtensionBase.h
+++ b/Engine/Plugins/Experimental/ClonerEffector/Source/ClonerEffector/Public/Cloner/Extensions/CEClonerExtensionBase.h
@@ -30,6 +30,7 @@ public:
 		, ExtensionPriority(InExtensionPriority)
 	{}
 
+	UFUNCTION(BlueprintPure, Category="Cloner")
 	FName GetExtensionName() const
 	{
 		return ExtensionName;
@@ -59,6 +60,7 @@ public:
 	/** Disable this extension */
 	void DeactivateExtension();
 
+	UFUNCTION(BlueprintPure, Category="Cloner")
 	bool IsExtensionActive() const
 	{
 		return bExtensionActive;
diff --git a/Engine/Plugins/Experimental/ClonerEffector/Source/ClonerEffector/Public/Cloner/Extensions/CEClonerLifetimeExtension.h b/Engine/Plugins/Experimental/ClonerEffector/Source/ClonerEffector/Public/Cloner/Extensions/CEClonerLifetimeExtension.h
--- a/Engine/Plugins/Experimental/ClonerEffector/Source/ClonerEffector/Public/Cloner/Extensions/CEClonerLifetimeExtension.h
+++ b/Engine/Plugins/Experimental/ClonerEffector/Source/ClonerEffector/Public/Cloner/Extensions/CEClonerLifetimeExtension.h
@@ -7,6 +7,7 @@
 #include "Curves/RichCurve.h"
 #include "CEClonerLifetimeExtension.generated.h"
 
+class UCurveFloat;
 class UNiagaraDataInterfaceCurve;
 
 /** Extension dealing with clones lifetime options */
@@ -54,10 +55,11 @@ public:
 		return bLifetimeScaleEnabled;
 	}
 
-	UFUNCTION()
+	UFUNCTION(BlueprintCallable, Category="Cloner")
+	void SetLifetimeScaleCurve(UCurveFloat* InCurve);
+
 	CLONEREFFECTOR_API void SetLifetimeScaleCurve(const FRichCurve& InCurve);
 
-	UFUNCTION()
 	const FRichCurve& GetLifetimeScaleCurve() const
 	{
 		return LifetimeScaleCurve;
@@ -82,19 +84,19 @@ protected:
 	void OnLifetimeScaleCurveChanged();
 
 	/** Do we destroy the clones after a specific duration */
-	UPROPERTY(EditInstanceOnly, BlueprintReadWrite, Setter="SetLifetimeEnabled", Getter="GetLifetimeEnabled", DisplayName="Enabled", Category="Lifetime")
+	UPROPERTY(EditInstanceOnly, Setter="SetLifetimeEnabled", Getter="GetLifetimeEnabled", DisplayName="Enabled", Category="Lifetime")
 	bool bLifetimeEnabled = false;
 
 	/** Minimum lifetime for a clone */
-	UPROPERTY(EditInstanceOnly, BlueprintReadWrite, Setter="SetLifetimeMin", Getter="GetLifetimeMin", DisplayName="Min", Category="Lifetime", meta=(ClampMin="0", EditCondition="bLifetimeEnabled", EditConditionHides))
+	UPROPERTY(EditInstanceOnly, Setter, Getter, DisplayName="Min", Category="Lifetime", meta=(ClampMin="0", EditCondition="bLifetimeEnabled", EditConditionHides))
 	float LifetimeMin = 0.25f;
 
 	/** Maximum lifetime for a clone */
-	UPROPERTY(EditInstanceOnly, BlueprintReadWrite, Setter="SetLifetimeMax", Getter="GetLifetimeMax", DisplayName="Max", Category="Lifetime", meta=(ClampMin="0", EditCondition="bLifetimeEnabled", EditConditionHides))
+	UPROPERTY(EditInstanceOnly, Setter, Getter, DisplayName="Max", Category="Lifetime", meta=(ClampMin="0", EditCondition="bLifetimeEnabled", EditConditionHides))
 	float LifetimeMax = 1.5f;
 
 	/** Enable scale by lifetime */
-	UPROPERTY(EditInstanceOnly, BlueprintReadWrite, Setter="SetLifetimeScaleEnabled", Getter="GetLifetimeScaleEnabled", DisplayName="ScaleEnabled", Category="Lifetime", meta=(EditCondition="bLifetimeEnabled", EditConditionHides))
+	UPROPERTY(EditInstanceOnly, Setter="SetLifetimeScaleEnabled", Getter="GetLifetimeScaleEnabled", DisplayName="ScaleEnabled", Category="Lifetime", meta=(EditCondition="bLifetimeEnabled", EditConditionHides))
 	bool bLifetimeScaleEnabled = false;
 
 	/** Used to expose the scale curve editor in details panel */
@@ -102,7 +104,7 @@ protected:
 	TWeakObjectPtr<UNiagaraDataInterfaceCurve> LifetimeScaleCurveDIWeak;
 
 	/** Rich curve used in the data interface for the lifetime */
-	UPROPERTY(Setter="SetLifetimeScaleCurve", Getter="GetLifetimeScaleCurve")
+	UPROPERTY(Setter, Getter)
 	FRichCurve LifetimeScaleCurve;
 
 private:
diff --git a/Engine/Plugins/Experimental/ClonerEffector/Source/ClonerEffector/Public/Cloner/Extensions/CEClonerMeshRendererExtension.h b/Engine/Plugins/Experimental/ClonerEffector/Source/ClonerEffector/Public/Cloner/Extensions/CEClonerMeshRendererExtension.h
--- a/Engine/Plugins/Experimental/ClonerEffector/Source/ClonerEffector/Public/Cloner/Extensions/CEClonerMeshRendererExtension.h
+++ b/Engine/Plugins/Experimental/ClonerEffector/Source/ClonerEffector/Public/Cloner/Extensions/CEClonerMeshRendererExtension.h
@@ -52,10 +52,10 @@ public:
 	}
 
 	UFUNCTION(BlueprintCallable, Category="Cloner", meta=(DisplayName="SetDefaultMeshes"))
-	CLONEREFFECTOR_API void BP_SetDefaultMeshes(const TArray<UStaticMesh*>& InMeshes);
+	CLONEREFFECTOR_API void SetDefaultMeshes(const TArray<UStaticMesh*>& InMeshes);
 
 	UFUNCTION(BlueprintPure, Category="Cloner", meta=(DisplayName="GetDefaultMeshes"))
-	CLONEREFFECTOR_API TArray<UStaticMesh*> BP_GetDefaultMeshes() const;
+	CLONEREFFECTOR_API void GetDefaultMeshes(TArray<UStaticMesh*>& OutMeshes) const;
 
 	UFUNCTION(BlueprintCallable, Category="Cloner")
 	CLONEREFFECTOR_API void SetVisualizeEffectors(bool bInVisualize);
@@ -105,31 +105,31 @@ protected:
 	void OnOverrideMaterialOptionsChanged();
 
 	/** Indicates how we select the mesh to render on each clones */
-	UPROPERTY(EditInstanceOnly, BlueprintReadWrite, Setter="SetMeshRenderMode", Getter="GetMeshRenderMode", Category="Renderer")
+	UPROPERTY(EditInstanceOnly, Setter, Getter, Category="Renderer")
 	ECEClonerMeshRenderMode MeshRenderMode = ECEClonerMeshRenderMode::Iterate;
 
 	/** Mode to indicate how clones facing is determined */
-	UPROPERTY(EditInstanceOnly, BlueprintReadWrite, Setter="SetMeshFacingMode", Getter="GetMeshFacingMode", Category="Renderer")
+	UPROPERTY(EditInstanceOnly, Setter, Getter, Category="Renderer")
 	ENiagaraMeshFacingMode MeshFacingMode = ENiagaraMeshFacingMode::Default;
 
 	/** Whether clones cast shadow, disabling will result in better performance */
-	UPROPERTY(EditInstanceOnly, BlueprintReadWrite, Setter="SetMeshCastShadows", Getter="GetMeshCastShadows", Category="Renderer")
+	UPROPERTY(EditInstanceOnly, Setter="SetMeshCastShadows", Getter="GetMeshCastShadows", Category="Renderer")
 	bool bMeshCastShadows = true;
 
 	/** When nothing is attached to the cloner, these meshes are used as default */
-	UPROPERTY(EditInstanceOnly, BlueprintReadWrite, Setter="SetDefaultMeshes", Getter="GetDefaultMeshes", BlueprintSetter="BP_SetDefaultMeshes", BlueprintGetter="BP_GetDefaultMeshes", Category="Renderer")
+	UPROPERTY(EditInstanceOnly, Setter, Getter, Category="Renderer")
 	TArray<TObjectPtr<UStaticMesh>> DefaultMeshes;
 
 	/** Override materials to show effectors applied on this cloner based on their color property */
-	UPROPERTY(EditInstanceOnly, BlueprintReadWrite, Setter="SetVisualizeEffectors", Getter="GetVisualizeEffectors", Category="Renderer")
+	UPROPERTY(EditInstanceOnly, Setter="SetVisualizeEffectors", Getter="GetVisualizeEffectors", Category="Renderer")
 	bool bVisualizeEffectors = false;
 
 	/** Whether to override meshes materials with another material */
-	UPROPERTY(EditInstanceOnly, BlueprintReadWrite, Setter="SetUseOverrideMaterial", Getter="GetUseOverrideMaterial", Category="Renderer")
+	UPROPERTY(EditInstanceOnly, Setter="SetUseOverrideMaterial", Getter="GetUseOverrideMaterial", Category="Renderer")
 	bool bUseOverrideMaterial = false;
 
 	/** The override materials that will be set instead of meshes materials, bVisualizeEffectors must be disabled */
-	UPROPERTY(EditInstanceOnly, BlueprintReadWrite, Setter, Getter, Category="Renderer", meta=(EditCondition="bUseOverrideMaterial", EditConditionHides))
+	UPROPERTY(EditInstanceOnly, Setter, Getter, Category="Renderer", meta=(EditCondition="bUseOverrideMaterial", EditConditionHides))
 	TObjectPtr<UMaterialInterface> OverrideMaterial;
 
 private:
diff --git a/Engine/Plugins/Experimental/ClonerEffector/Source/ClonerEffector/Public/Cloner/Extensions/CEClonerProgressExtension.h b/Engine/Plugins/Experimental/ClonerEffector/Source/ClonerEffector/Public/Cloner/Extensions/CEClonerProgressExtension.h
--- a/Engine/Plugins/Experimental/ClonerEffector/Source/ClonerEffector/Public/Cloner/Extensions/CEClonerProgressExtension.h
+++ b/Engine/Plugins/Experimental/ClonerEffector/Source/ClonerEffector/Public/Cloner/Extensions/CEClonerProgressExtension.h
@@ -50,11 +50,11 @@ protected:
 	//~ End UCEClonerExtensionBase
 
 	/** Invert progress behaviour */
-	UPROPERTY(EditInstanceOnly, BlueprintReadWrite, Setter="SetInvertProgress", Getter="GetInvertProgress", Category="Progress")
+	UPROPERTY(EditInstanceOnly, Setter="SetInvertProgress", Getter="GetInvertProgress", Category="Progress")
 	bool bInvertProgress = false;
 
 	/** Changes visibility of instances based on the total count, 1.f = 100% = all instances visible */
-	UPROPERTY(EditInstanceOnly, BlueprintReadWrite, Setter, Getter, Category="Progress", meta=(ClampMin="0", ClampMax="1"))
+	UPROPERTY(EditInstanceOnly, Setter, Getter, Category="Progress", meta=(ClampMin="0", ClampMax="1"))
 	float Progress = 1.f;
 
 private:
diff --git a/Engine/Plugins/Experimental/ClonerEffector/Source/ClonerEffector/Public/Cloner/Extensions/CEClonerRangeExtension.h b/Engine/Plugins/Experimental/ClonerEffector/Source/ClonerEffector/Public/Cloner/Extensions/CEClonerRangeExtension.h
--- a/Engine/Plugins/Experimental/ClonerEffector/Source/ClonerEffector/Public/Cloner/Extensions/CEClonerRangeExtension.h
+++ b/Engine/Plugins/Experimental/ClonerEffector/Source/ClonerEffector/Public/Cloner/Extensions/CEClonerRangeExtension.h
@@ -131,38 +131,38 @@ protected:
 	//~ End UCEClonerExtensionBase
 
 	/** Use random range transforms for each clones */
-	UPROPERTY(EditInstanceOnly, BlueprintReadWrite, Setter="SetRangeEnabled", Getter="GetRangeEnabled", DisplayName="Enabled", Category="Range")
+	UPROPERTY(EditInstanceOnly, Setter="SetRangeEnabled", Getter="GetRangeEnabled", DisplayName="Enabled", Category="Range")
 	bool bRangeEnabled = false;
 
 	/** Mirrors offset and rotation values */
-	UPROPERTY(EditInstanceOnly, BlueprintReadWrite, Setter="SetRangeMirrored", Getter="GetRangeMirrored", DisplayName="Mirrored", Category="Range", meta=(EditCondition="bRangeEnabled", EditConditionHides))
+	UPROPERTY(EditInstanceOnly, Setter="SetRangeMirrored", Getter="GetRangeMirrored", DisplayName="Mirrored", Category="Range", meta=(EditCondition="bRangeEnabled", EditConditionHides))
 	bool bRangeMirrored = true;
 
-	UPROPERTY(EditInstanceOnly, BlueprintReadWrite, Setter="SetRangeOffsetMin", Getter="GetRangeOffsetMin", DisplayName="OffsetMin", Category="Range", meta=(EditCondition="bRangeEnabled && !bRangeMirrored", EditConditionHides))
+	UPROPERTY(EditInstanceOnly, Setter, Getter, DisplayName="OffsetMin", Category="Range", meta=(EditCondition="bRangeEnabled && !bRangeMirrored", EditConditionHides))
 	FVector RangeOffsetMin = FVector::ZeroVector;
 
-	UPROPERTY(EditInstanceOnly, BlueprintReadWrite, Setter="SetRangeOffsetMax", Getter="GetRangeOffsetMax", DisplayName="OffsetMax", Category="Range", meta=(EditCondition="bRangeEnabled", EditConditionHides))
+	UPROPERTY(EditInstanceOnly, Setter, Getter, DisplayName="OffsetMax", Category="Range", meta=(EditCondition="bRangeEnabled", EditConditionHides))
 	FVector RangeOffsetMax = FVector::ZeroVector;
 
-	UPROPERTY(EditInstanceOnly, BlueprintReadWrite, Setter="SetRangeRotationMin", Getter="GetRangeRotationMin", DisplayName="RotationMin", Category="Range", meta=(EditCondition="bRangeEnabled && !bRangeMirrored", EditConditionHides))
+	UPROPERTY(EditInstanceOnly, Setter, Getter, DisplayName="RotationMin", Category="Range", meta=(EditCondition="bRangeEnabled && !bRangeMirrored", EditConditionHides))
 	FRotator RangeRotationMin = FRotator::ZeroRotator;
 
-	UPROPERTY(EditInstanceOnly, BlueprintReadWrite, Setter="SetRangeRotationMax", Getter="GetRangeRotationMax", DisplayName="RotationMax", Category="Range", meta=(EditCondition="bRangeEnabled", EditConditionHides))
+	UPROPERTY(EditInstanceOnly, Setter, Getter, DisplayName="RotationMax", Category="Range", meta=(EditCondition="bRangeEnabled", EditConditionHides))
 	FRotator RangeRotationMax = FRotator::ZeroRotator;
 
-	UPROPERTY(EditInstanceOnly, BlueprintReadWrite, Setter="SetRangeScaleUniform", Getter="GetRangeScaleUniform", DisplayName="ScaleUniformEnabled", Category="Range", meta=(EditCondition="bRangeEnabled", EditConditionHides))
+	UPROPERTY(EditInstanceOnly, Setter="SetRangeScaleUniform", Getter="GetRangeScaleUniform", DisplayName="ScaleUniformEnabled", Category="Range", meta=(EditCondition="bRangeEnabled", EditConditionHides))
 	bool bRangeScaleUniform = true;
 
-	UPROPERTY(EditInstanceOnly, BlueprintReadWrite, Setter="SetRangeScaleMin", Getter="GetRangeScaleMin", DisplayName="ScaleMin", Category="Range", meta=(MotionDesignVectorWidget, AllowPreserveRatio="XYZ", Delta="0.0001", ClampMin="0", EditCondition="bRangeEnabled && !bRangeScaleUniform", EditConditionHides))
+	UPROPERTY(EditInstanceOnly, Setter, Getter, DisplayName="ScaleMin", Category="Range", meta=(MotionDesignVectorWidget, AllowPreserveRatio="XYZ", Delta="0.0001", ClampMin="0", EditCondition="bRangeEnabled && !bRangeScaleUniform", EditConditionHides))
 	FVector RangeScaleMin = FVector::OneVector;
 
-	UPROPERTY(EditInstanceOnly, BlueprintReadWrite, Setter="SetRangeScaleMax", Getter="GetRangeScaleMax", DisplayName="ScaleMax", Category="Range", meta=(MotionDesignVectorWidget, AllowPreserveRatio="XYZ", Delta="0.0001", ClampMin="0", EditCondition="bRangeEnabled && !bRangeScaleUniform", EditConditionHides))
+	UPROPERTY(EditInstanceOnly, Setter, Getter, DisplayName="ScaleMax", Category="Range", meta=(MotionDesignVectorWidget, AllowPreserveRatio="XYZ", Delta="0.0001", ClampMin="0", EditCondition="bRangeEnabled && !bRangeScaleUniform", EditConditionHides))
 	FVector RangeScaleMax = FVector::OneVector;
 
-	UPROPERTY(EditInstanceOnly, BlueprintReadWrite, Setter="SetRangeScaleUniformMin", Getter="GetRangeScaleUniformMin", DisplayName="ScaleMin", Category="Range", meta=(Delta="0.0001", ClampMin="0", EditCondition="bRangeEnabled && bRangeScaleUniform", EditConditionHides))
+	UPROPERTY(EditInstanceOnly, Setter, Getter, DisplayName="ScaleMin", Category="Range", meta=(Delta="0.0001", ClampMin="0", EditCondition="bRangeEnabled && bRangeScaleUniform", EditConditionHides))
 	float RangeScaleUniformMin = 1.f;
 
-	UPROPERTY(EditInstanceOnly, BlueprintReadWrite, Setter="SetRangeScaleUniformMax", Getter="GetRangeScaleUniformMax", DisplayName="ScaleMax", Category="Range", meta=(Delta="0.0001", ClampMin="0", EditCondition="bRangeEnabled && bRangeScaleUniform", EditConditionHides))
+	UPROPERTY(EditInstanceOnly, Setter, Getter, DisplayName="ScaleMax", Category="Range", meta=(Delta="0.0001", ClampMin="0", EditCondition="bRangeEnabled && bRangeScaleUniform", EditConditionHides))
 	float RangeScaleUniformMax = 1.f;
 
 private:
diff --git a/Engine/Plugins/Experimental/ClonerEffector/Source/ClonerEffector/Public/Cloner/Extensions/CEClonerStepExtension.h b/Engine/Plugins/Experimental/ClonerEffector/Source/ClonerEffector/Public/Cloner/Extensions/CEClonerStepExtension.h
--- a/Engine/Plugins/Experimental/ClonerEffector/Source/ClonerEffector/Public/Cloner/Extensions/CEClonerStepExtension.h
+++ b/Engine/Plugins/Experimental/ClonerEffector/Source/ClonerEffector/Public/Cloner/Extensions/CEClonerStepExtension.h
@@ -68,20 +68,20 @@ protected:
 	//~ End UCEClonerExtensionBase
 
 	/** Enable steps to add delta variation on each clone instance */
-	UPROPERTY(EditInstanceOnly, BlueprintReadWrite, Setter="SetDeltaStepEnabled", Getter="GetDeltaStepEnabled", DisplayName="Enabled", Category="Step")
+	UPROPERTY(EditInstanceOnly, Setter="SetDeltaStepEnabled", Getter="GetDeltaStepEnabled", DisplayName="Enabled", Category="Step")
 	bool bDeltaStepEnabled = false;
 
 	/** Amount of position difference between one step and the next one */
-	UPROPERTY(EditInstanceOnly, BlueprintReadWrite, Setter="SetDeltaStepPosition", Getter="GetDeltaStepPosition", DisplayName="Position", Category="Step", meta=(EditCondition="bDeltaStepEnabled", EditConditionHides))
-	FVector DeltaStepPosition = FVector(0.f);
+	UPROPERTY(EditInstanceOnly, Setter, Getter, DisplayName="Position", Category="Step", meta=(EditCondition="bDeltaStepEnabled", EditConditionHides))
+	FVector DeltaStepPosition = FVector::ZeroVector;
 
 	/** Amount of rotation difference between one step and the next one */
-	UPROPERTY(EditInstanceOnly, BlueprintReadWrite, Setter="SetDeltaStepRotation", Getter="GetDeltaStepRotation", DisplayName="Rotation", Category="Step", meta=(EditCondition="bDeltaStepEnabled", EditConditionHides))
-	FRotator DeltaStepRotation = FRotator(0.f);
+	UPROPERTY(EditInstanceOnly, Setter, Getter, DisplayName="Rotation", Category="Step", meta=(EditCondition="bDeltaStepEnabled", EditConditionHides))
+	FRotator DeltaStepRotation = FRotator::ZeroRotator;
 
 	/** Amount of scale difference between one step and the next one */
-	UPROPERTY(EditInstanceOnly, BlueprintReadWrite, Setter="SetDeltaStepScale", Getter="GetDeltaStepScale", DisplayName="Scale", Category="Step", meta=(MotionDesignVectorWidget, AllowPreserveRatio="XYZ", Delta="0.0001", EditCondition="bDeltaStepEnabled", EditConditionHides))
-	FVector DeltaStepScale = FVector(0.f);
+	UPROPERTY(EditInstanceOnly, Setter, Getter, DisplayName="Scale", Category="Step", meta=(MotionDesignVectorWidget, AllowPreserveRatio="XYZ", Delta="0.0001", EditCondition="bDeltaStepEnabled", EditConditionHides))
+	FVector DeltaStepScale = FVector::ZeroVector;
 
 private:
 #if WITH_EDITOR
diff --git a/Engine/Plugins/Experimental/ClonerEffector/Source/ClonerEffector/Public/Cloner/Layouts/CEClonerCircleLayout.h b/Engine/Plugins/Experimental/ClonerEffector/Source/ClonerEffector/Public/Cloner/Layouts/CEClonerCircleLayout.h
--- a/Engine/Plugins/Experimental/ClonerEffector/Source/ClonerEffector/Public/Cloner/Layouts/CEClonerCircleLayout.h
+++ b/Engine/Plugins/Experimental/ClonerEffector/Source/ClonerEffector/Public/Cloner/Layouts/CEClonerCircleLayout.h
@@ -105,28 +105,28 @@ protected:
 	virtual void OnLayoutParametersChanged(UCEClonerComponent* InComponent) override;
 	//~ End UCEClonerLayoutBase
 
-	UPROPERTY(EditInstanceOnly, BlueprintReadWrite, Setter="SetCount", Getter="GetCount", Category="Layout")
+	UPROPERTY(EditInstanceOnly, Setter, Getter, Category="Layout")
 	int32 Count = 3 * 3;
 
-	UPROPERTY(EditInstanceOnly, BlueprintReadWrite, Setter="SetRadius", Getter="GetRadius", Category="Layout")
+	UPROPERTY(EditInstanceOnly, Setter, Getter, Category="Layout")
 	float Radius = 200.f;
 
-	UPROPERTY(EditInstanceOnly, BlueprintReadWrite, Setter="SetAngleStart", Getter="GetAngleStart", Category="Layout")
+	UPROPERTY(EditInstanceOnly, Setter, Getter, Category="Layout")
 	float AngleStart = 0.f;
 
-	UPROPERTY(EditInstanceOnly, BlueprintReadWrite, Setter="SetAngleRatio", Getter="GetAngleRatio", Category="Layout")
+	UPROPERTY(EditInstanceOnly, Setter, Getter, Category="Layout")
 	float AngleRatio = 1.f;
 
-	UPROPERTY(EditInstanceOnly, BlueprintReadWrite, Setter="SetOrientMesh", Getter="GetOrientMesh", Category="Layout")
+	UPROPERTY(EditInstanceOnly, Setter="SetOrientMesh", Getter="GetOrientMesh", Category="Layout")
 	bool bOrientMesh = true;
 
-	UPROPERTY(EditInstanceOnly, BlueprintReadWrite, Setter="SetPlane", Getter="GetPlane", Category="Layout")
+	UPROPERTY(EditInstanceOnly, Setter, Getter, Category="Layout")
 	ECEClonerPlane Plane = ECEClonerPlane::XY;
 
-	UPROPERTY(EditInstanceOnly, BlueprintReadWrite, Setter="SetRotation", Getter="GetRotation", Category="Layout", meta=(EditCondition="Plane == ECEClonerPlane::Custom", EditConditionHides))
+	UPROPERTY(EditInstanceOnly, Setter, Getter, Category="Layout", meta=(EditCondition="Plane == ECEClonerPlane::Custom", EditConditionHides))
 	FRotator Rotation = FRotator(0.f, 0.f, 0.f);
 
-	UPROPERTY(EditInstanceOnly, BlueprintReadWrite, Setter="SetScale", Getter="GetScale", Category="Layout", meta=(ClampMin="0", MotionDesignVectorWidget, AllowPreserveRatio="XYZ", Delta="0.01"))
+	UPROPERTY(EditInstanceOnly, Setter, Getter, Category="Layout", meta=(ClampMin="0", MotionDesignVectorWidget, AllowPreserveRatio="XYZ", Delta="0.01"))
 	FVector Scale = FVector(1.f, 1.f, 1.f);
 
 private:
diff --git a/Engine/Plugins/Experimental/ClonerEffector/Source/ClonerEffector/Public/Cloner/Layouts/CEClonerCylinderLayout.h b/Engine/Plugins/Experimental/ClonerEffector/Source/ClonerEffector/Public/Cloner/Layouts/CEClonerCylinderLayout.h
--- a/Engine/Plugins/Experimental/ClonerEffector/Source/ClonerEffector/Public/Cloner/Layouts/CEClonerCylinderLayout.h
+++ b/Engine/Plugins/Experimental/ClonerEffector/Source/ClonerEffector/Public/Cloner/Layouts/CEClonerCylinderLayout.h
@@ -123,34 +123,34 @@ protected:
 	virtual void OnLayoutParametersChanged(UCEClonerComponent* InComponent) override;
 	//~ End UCEClonerLayoutBase
 
-	UPROPERTY(EditInstanceOnly, BlueprintReadWrite, Setter="SetBaseCount", Getter="GetBaseCount", Category="Layout")
+	UPROPERTY(EditInstanceOnly, Setter, Getter, Category="Layout")
 	int32 BaseCount = 3 * 3;
 
-	UPROPERTY(EditInstanceOnly, BlueprintReadWrite, Setter="SetHeightCount", Getter="GetHeightCount", Category="Layout")
+	UPROPERTY(EditInstanceOnly, Setter, Getter, Category="Layout")
 	int32 HeightCount = 3;
 
-	UPROPERTY(EditInstanceOnly, BlueprintReadWrite, Setter="SetHeight", Getter="GetHeight", Category="Layout")
+	UPROPERTY(EditInstanceOnly, Setter, Getter, Category="Layout")
 	float Height = 400.f;
 
-	UPROPERTY(EditInstanceOnly, BlueprintReadWrite, Setter="SetRadius", Getter="GetRadius", Category="Layout")
+	UPROPERTY(EditInstanceOnly, Setter, Getter, Category="Layout")
 	float Radius = 200.f;
 
-	UPROPERTY(EditInstanceOnly, BlueprintReadWrite, Setter="SetAngleStart", Getter="GetAngleStart", Category="Layout")
+	UPROPERTY(EditInstanceOnly, Setter, Getter, Category="Layout")
 	float AngleStart = 0.f;
 
-	UPROPERTY(EditInstanceOnly, BlueprintReadWrite, Setter="SetAngleRatio", Getter="GetAngleRatio", Category="Layout")
+	UPROPERTY(EditInstanceOnly, Setter, Getter, Category="Layout")
 	float AngleRatio = 1.f;
 
-	UPROPERTY(EditInstanceOnly, BlueprintReadWrite, Setter="SetOrientMesh", Getter="GetOrientMesh", Category="Layout")
+	UPROPERTY(EditInstanceOnly, Setter="SetOrientMesh", Getter="GetOrientMesh", Category="Layout")
 	bool bOrientMesh = true;
 
-	UPROPERTY(EditInstanceOnly, BlueprintReadWrite, Setter="SetPlane", Getter="GetPlane", Category="Layout")
+	UPROPERTY(EditInstanceOnly, Setter, Getter, Category="Layout")
 	ECEClonerPlane Plane = ECEClonerPlane::XY;
 
-	UPROPERTY(EditInstanceOnly, BlueprintReadWrite, Setter="SetRotation", Getter="GetRotation", Category="Layout", meta=(EditCondition="Plane == ECEClonerPlane::Custom", EditConditionHides))
+	UPROPERTY(EditInstanceOnly, Setter, Getter, Category="Layout", meta=(EditCondition="Plane == ECEClonerPlane::Custom", EditConditionHides))
 	FRotator Rotation = FRotator(0.f, 0.f, 0.f);
 
-	UPROPERTY(EditInstanceOnly, BlueprintReadWrite, Setter="SetScale", Getter="GetScale", Category="Layout", meta=(ClampMin="0", MotionDesignVectorWidget, AllowPreserveRatio="XYZ", Delta="0.01"))
+	UPROPERTY(EditInstanceOnly, Setter, Getter, Category="Layout", meta=(ClampMin="0", MotionDesignVectorWidget, AllowPreserveRatio="XYZ", Delta="0.01"))
 	FVector Scale = FVector(1.f, 1.f, 1.f);
 
 private:
diff --git a/Engine/Plugins/Experimental/ClonerEffector/Source/ClonerEffector/Public/Cloner/Layouts/CEClonerGridLayout.h b/Engine/Plugins/Experimental/ClonerEffector/Source/ClonerEffector/Public/Cloner/Layouts/CEClonerGridLayout.h
--- a/Engine/Plugins/Experimental/ClonerEffector/Source/ClonerEffector/Public/Cloner/Layouts/CEClonerGridLayout.h
+++ b/Engine/Plugins/Experimental/ClonerEffector/Source/ClonerEffector/Public/Cloner/Layouts/CEClonerGridLayout.h
@@ -134,28 +134,28 @@ protected:
 	virtual void OnLayoutParametersChanged(UCEClonerComponent* InComponent) override;
 	//~ End UCEClonerLayoutBase
 
-	UPROPERTY(EditInstanceOnly, BlueprintReadWrite, Setter="SetCountX", Getter="GetCountX", Category="Layout")
+	UPROPERTY(EditInstanceOnly, Setter, Getter, Category="Layout")
 	int32 CountX = 3;
 
-	UPROPERTY(EditInstanceOnly, BlueprintReadWrite, Setter="SetCountY", Getter="GetCountY", Category="Layout")
+	UPROPERTY(EditInstanceOnly, Setter, Getter, Category="Layout")
 	int32 CountY = 3;
 
-	UPROPERTY(EditInstanceOnly, BlueprintReadWrite, Setter="SetCountZ", Getter="GetCountZ", Category="Layout")
+	UPROPERTY(EditInstanceOnly, Setter, Getter, Category="Layout")
 	int32 CountZ = 3;
 
-	UPROPERTY(EditInstanceOnly, BlueprintReadWrite, Setter="SetSpacingX", Getter="GetSpacingX", Category="Layout")
+	UPROPERTY(EditInstanceOnly, Setter, Getter, Category="Layout")
 	float SpacingX = 105.f;
 
-	UPROPERTY(EditInstanceOnly, BlueprintReadWrite, Setter="SetSpacingY", Getter="GetSpacingY", Category="Layout")
+	UPROPERTY(EditInstanceOnly, Setter, Getter, Category="Layout")
 	float SpacingY = 105.f;
 
-	UPROPERTY(EditInstanceOnly, BlueprintReadWrite, Setter="SetSpacingZ", Getter="GetSpacingZ", Category="Layout")
+	UPROPERTY(EditInstanceOnly, Setter, Getter, Category="Layout")
 	float SpacingZ = 105.f;
 
-	UPROPERTY(EditInstanceOnly, BlueprintReadWrite, Setter, Getter, Category="Layout", meta=(UIMin="0", UIMax="100"))
+	UPROPERTY(EditInstanceOnly, Setter, Getter, Category="Layout", meta=(UIMin="0", UIMax="100"))
 	float TwistFactor = 0.f;
 
-	UPROPERTY(EditInstanceOnly, BlueprintReadWrite, Setter, Getter, Category="Layout", meta=(InvalidEnumValues="Custom"))
+	UPROPERTY(EditInstanceOnly, Setter, Getter, Category="Layout", meta=(InvalidEnumValues="Custom"))
 	ECEClonerAxis TwistAxis = ECEClonerAxis::Y;
 
 private:
diff --git a/Engine/Plugins/Experimental/ClonerEffector/Source/ClonerEffector/Public/Cloner/Layouts/CEClonerHoneycombLayout.h b/Engine/Plugins/Experimental/ClonerEffector/Source/ClonerEffector/Public/Cloner/Layouts/CEClonerHoneycombLayout.h
--- a/Engine/Plugins/Experimental/ClonerEffector/Source/ClonerEffector/Public/Cloner/Layouts/CEClonerHoneycombLayout.h
+++ b/Engine/Plugins/Experimental/ClonerEffector/Source/ClonerEffector/Public/Cloner/Layouts/CEClonerHoneycombLayout.h
@@ -116,31 +116,31 @@ protected:
 	virtual void OnLayoutParametersChanged(UCEClonerComponent* InComponent) override;
 	//~ End UCEClonerLayoutBase
 
-	UPROPERTY(EditInstanceOnly, BlueprintReadWrite, Setter="SetPlane", Getter="GetPlane", Category="Layout", meta=(InvalidEnumValues="Custom"))
+	UPROPERTY(EditInstanceOnly, Setter, Getter, Category="Layout", meta=(InvalidEnumValues="Custom"))
 	ECEClonerPlane Plane = ECEClonerPlane::XY;
 
-	UPROPERTY(EditInstanceOnly, BlueprintReadWrite, Setter="SetWidthCount", Getter="GetWidthCount", Category="Layout")
+	UPROPERTY(EditInstanceOnly, Setter, Getter, Category="Layout")
 	int32 WidthCount = 3;
 
-	UPROPERTY(EditInstanceOnly, BlueprintReadWrite, Setter="SetHeightCount", Getter="GetHeightCount", Category="Layout")
+	UPROPERTY(EditInstanceOnly, Setter, Getter, Category="Layout")
 	int32 HeightCount = 3;
 
-	UPROPERTY(EditInstanceOnly, BlueprintReadWrite, Setter="SetWidthOffset", Getter="GetWidthOffset", Category="Layout")
+	UPROPERTY(EditInstanceOnly, Setter, Getter, Category="Layout")
 	float WidthOffset = 0.f;
 
-	UPROPERTY(EditInstanceOnly, BlueprintReadWrite, Setter="SetHeightOffset", Getter="GetHeightOffset", Category="Layout")
+	UPROPERTY(EditInstanceOnly, Setter, Getter, Category="Layout")
 	float HeightOffset = 0.5f;
 
-	UPROPERTY(EditInstanceOnly, BlueprintReadWrite, Setter="SetWidthSpacing", Getter="GetWidthSpacing", Category="Layout")
+	UPROPERTY(EditInstanceOnly, Setter, Getter, Category="Layout")
 	float WidthSpacing = 105.f;
 
-	UPROPERTY(EditInstanceOnly, BlueprintReadWrite, Setter="SetHeightSpacing", Getter="GetHeightSpacing", Category="Layout")
+	UPROPERTY(EditInstanceOnly, Setter, Getter, Category="Layout")
 	float HeightSpacing = 105.f;
 
-	UPROPERTY(EditInstanceOnly, BlueprintReadWrite, Setter, Getter, Category="Layout", meta=(UIMin="0", UIMax="100"))
+	UPROPERTY(EditInstanceOnly, Setter, Getter, Category="Layout", meta=(UIMin="0", UIMax="100"))
 	float TwistFactor = 0.f;
 
-	UPROPERTY(EditInstanceOnly, BlueprintReadWrite, Setter, Getter, Category="Layout", meta=(InvalidEnumValues="Custom"))
+	UPROPERTY(EditInstanceOnly, Setter, Getter, Category="Layout", meta=(InvalidEnumValues="Custom"))
 	ECEClonerAxis TwistAxis = ECEClonerAxis::Y;
 
 private:
diff --git a/Engine/Plugins/Experimental/ClonerEffector/Source/ClonerEffector/Public/Cloner/Layouts/CEClonerLayoutBase.h b/Engine/Plugins/Experimental/ClonerEffector/Source/ClonerEffector/Public/Cloner/Layouts/CEClonerLayoutBase.h
--- a/Engine/Plugins/Experimental/ClonerEffector/Source/ClonerEffector/Public/Cloner/Layouts/CEClonerLayoutBase.h
+++ b/Engine/Plugins/Experimental/ClonerEffector/Source/ClonerEffector/Public/Cloner/Layouts/CEClonerLayoutBase.h
@@ -25,6 +25,8 @@ class UCEClonerLayoutBase : public UObject
 {
 	GENERATED_BODY()
 
+	DECLARE_MULTICAST_DELEGATE_TwoParams(FOnClonerLayoutLoaded, UCEClonerLayoutBase* /** InLayout */, bool /** bInSuccess */)
+
 public:
 	static constexpr TCHAR LayoutBaseAssetPath[] = TEXT("/Script/Niagara.NiagaraSystem'/ClonerEffector/Systems/NS_ClonerBase.NS_ClonerBase'");
 
@@ -37,7 +39,7 @@ public:
 		, LayoutAssetPath(InLayoutAssetPath)
 	{}
 
-	UFUNCTION(BlueprintPure, Category="Cloner|Layout")
+	UFUNCTION(BlueprintPure, Category="Cloner")
 	FName GetLayoutName() const
 	{
 		return LayoutName;
@@ -58,13 +60,12 @@ public:
 		return MeshRenderer;
 	}
 
-	TMulticastDelegateRegistration<void(UCEClonerLayoutBase*, bool)>& OnLayoutLoadedDelegate()
+	FOnClonerLayoutLoaded::RegistrationType& OnLayoutLoadedDelegate()
 	{
 		return OnClonerLayoutLoadedDelegate;
 	}
 
 	/** Get the cloner component using this layout */
-	UFUNCTION(BlueprintPure, Category="Cloner|Layout")
 	UCEClonerComponent* GetClonerComponent() const;
 
 	/** Get the actor using this layout */
@@ -77,8 +78,7 @@ public:
 	bool IsLayoutValid() const;
 
 	/** Is this layout system cached and ready to be used */
-	UFUNCTION(BlueprintPure, Category="Cloner|Layout")
-	CLONEREFFECTOR_API bool IsLayoutLoaded() const;
+	bool IsLayoutLoaded() const;
 
 	/** Load this layout system if not already loaded */
 	void LoadLayout();
@@ -137,7 +137,6 @@ protected:
 	void OnLayoutPropertyChanged();
 
 private:
-	DECLARE_MULTICAST_DELEGATE_TwoParams(FOnClonerLayoutLoaded, UCEClonerLayoutBase* /** InLayout */, bool /** bInSuccess */)
 	FOnClonerLayoutLoaded OnClonerLayoutLoadedDelegate;
 
 	void OnSystemPackageLoaded(const FName& InName, UPackage* InPackage, EAsyncLoadingResult::Type InResult);
@@ -146,6 +145,7 @@ private:
 	void BindCleanupDelegates();
 	void UnbindCleanupDelegates() const;
 
+	/** When level is unloaded or world cleaned up, deactivate and unload layout */
 	void OnWorldCleanup(UWorld* InWorld, bool bInSessionEnded, bool bInCleanupResources);
 	void OnLevelCleanup();
 
diff --git a/Engine/Plugins/Experimental/ClonerEffector/Source/ClonerEffector/Public/Cloner/Layouts/CEClonerLineLayout.h b/Engine/Plugins/Experimental/ClonerEffector/Source/ClonerEffector/Public/Cloner/Layouts/CEClonerLineLayout.h
--- a/Engine/Plugins/Experimental/ClonerEffector/Source/ClonerEffector/Public/Cloner/Layouts/CEClonerLineLayout.h
+++ b/Engine/Plugins/Experimental/ClonerEffector/Source/ClonerEffector/Public/Cloner/Layouts/CEClonerLineLayout.h
@@ -78,19 +78,19 @@ protected:
 	virtual void OnLayoutParametersChanged(UCEClonerComponent* InComponent) override;
 	//~ End UCEClonerLayoutBase
 
-	UPROPERTY(EditInstanceOnly, BlueprintReadWrite, Setter="SetCount", Getter="GetCount", Category="Layout")
+	UPROPERTY(EditInstanceOnly, Setter, Getter, Category="Layout")
     int32 Count = 10;
 
-	UPROPERTY(EditInstanceOnly, BlueprintReadWrite, Setter="SetSpacing", Getter="GetSpacing", Category="Layout")
+	UPROPERTY(EditInstanceOnly, Setter, Getter, Category="Layout")
 	float Spacing = 105.f;
 
-	UPROPERTY(EditInstanceOnly, BlueprintReadWrite, Setter="SetAxis", Getter="GetAxis", Category="Layout")
+	UPROPERTY(EditInstanceOnly, Setter, Getter, Category="Layout")
 	ECEClonerAxis Axis = ECEClonerAxis::Y;
 
-	UPROPERTY(EditInstanceOnly, BlueprintReadWrite, Setter="SetDirection", Getter="GetDirection", Category="Layout", meta=(ClampMin="0", ClampMax="1", EditCondition="Axis == ECEClonerAxis::Custom", EditConditionHides))
+	UPROPERTY(EditInstanceOnly, Setter, Getter, Category="Layout", meta=(ClampMin="0", ClampMax="1", EditCondition="Axis == ECEClonerAxis::Custom", EditConditionHides))
 	FVector Direction = FVector::YAxisVector;
 
-	UPROPERTY(EditInstanceOnly, BlueprintReadWrite, Setter="SetRotation", Getter="GetRotation", Category="Layout")
+	UPROPERTY(EditInstanceOnly, Setter, Getter, Category="Layout")
 	FRotator Rotation = FRotator(0.f);
 
 private:
diff --git a/Engine/Plugins/Experimental/ClonerEffector/Source/ClonerEffector/Public/Cloner/Layouts/CEClonerMeshLayout.h b/Engine/Plugins/Experimental/ClonerEffector/Source/ClonerEffector/Public/Cloner/Layouts/CEClonerMeshLayout.h
--- a/Engine/Plugins/Experimental/ClonerEffector/Source/ClonerEffector/Public/Cloner/Layouts/CEClonerMeshLayout.h
+++ b/Engine/Plugins/Experimental/ClonerEffector/Source/ClonerEffector/Public/Cloner/Layouts/CEClonerMeshLayout.h
@@ -89,16 +89,16 @@ protected:
 
 	void OnSampleMeshTransformed(USceneComponent* InComponent, EUpdateTransformFlags InFlags, ETeleportType InType);
 
-	UPROPERTY(EditInstanceOnly, BlueprintReadWrite, Setter, Getter, Category="Layout")
+	UPROPERTY(EditInstanceOnly, Setter, Getter, Category="Layout")
 	ECEClonerMeshAsset Asset = ECEClonerMeshAsset::StaticMesh;
 
-	UPROPERTY(EditInstanceOnly, BlueprintReadWrite, Setter, Getter, DisplayName="Sample Actor", Category="Layout")
+	UPROPERTY(EditInstanceOnly, Setter, Getter, DisplayName="Sample Actor", Category="Layout")
 	TWeakObjectPtr<AActor> SampleActorWeak;
 
-	UPROPERTY(EditInstanceOnly, BlueprintReadWrite, Setter, Getter, Category="Layout")
+	UPROPERTY(EditInstanceOnly, Setter, Getter, Category="Layout")
 	int32 Count = 3 * 3 * 3;
 
-	UPROPERTY(EditInstanceOnly, BlueprintReadWrite, Setter, Getter, Category="Layout")
+	UPROPERTY(EditInstanceOnly, Setter, Getter, Category="Layout")
 	ECEClonerMeshSampleData SampleData = ECEClonerMeshSampleData::Vertices;
 
 private:
diff --git a/Engine/Plugins/Experimental/ClonerEffector/Source/ClonerEffector/Public/Cloner/Layouts/CEClonerSphereRandomLayout.h b/Engine/Plugins/Experimental/ClonerEffector/Source/ClonerEffector/Public/Cloner/Layouts/CEClonerSphereRandomLayout.h
--- a/Engine/Plugins/Experimental/ClonerEffector/Source/ClonerEffector/Public/Cloner/Layouts/CEClonerSphereRandomLayout.h
+++ b/Engine/Plugins/Experimental/ClonerEffector/Source/ClonerEffector/Public/Cloner/Layouts/CEClonerSphereRandomLayout.h
@@ -104,28 +104,28 @@ protected:
 	virtual void OnLayoutParametersChanged(UCEClonerComponent* InComponent) override;
 	//~ End UCEClonerLayoutBase
 
-	UPROPERTY(EditInstanceOnly, BlueprintReadWrite, Setter="SetCount", Getter="GetCount", Category="Layout")
+	UPROPERTY(EditInstanceOnly, Setter, Getter, Category="Layout")
 	int32 Count = 30;
 
-	UPROPERTY(EditInstanceOnly, BlueprintReadWrite, Setter="SetRadius", Getter="GetRadius", Category="Layout")
+	UPROPERTY(EditInstanceOnly, Setter, Getter, Category="Layout")
 	float Radius = 200.f;
 
-	UPROPERTY(EditInstanceOnly, BlueprintReadWrite, Setter="SetDistribution", Getter="GetDistribution", Category="Layout", meta=(ClampMin="0", ClampMax="1"))
+	UPROPERTY(EditInstanceOnly, Setter, Getter, Category="Layout", meta=(ClampMin="0", ClampMax="1"))
 	float Distribution = 1.f;
 
-	UPROPERTY(EditInstanceOnly, BlueprintReadWrite, Setter="SetLongitude", Getter="GetLongitude", Category="Layout", meta=(ClampMin="0", ClampMax="1"))
+	UPROPERTY(EditInstanceOnly, Setter, Getter, Category="Layout", meta=(ClampMin="0", ClampMax="1"))
 	float Longitude = 1.f;
 
-	UPROPERTY(EditInstanceOnly, BlueprintReadWrite, Setter="SetLatitude", Getter="GetLatitude", Category="Layout", meta=(ClampMin="0", ClampMax="1"))
+	UPROPERTY(EditInstanceOnly, Setter, Getter, Category="Layout", meta=(ClampMin="0", ClampMax="1"))
 	float Latitude = 1.f;
 
-	UPROPERTY(EditInstanceOnly, BlueprintReadWrite, Setter="SetOrientMesh", Getter="GetOrientMesh", Category="Layout")
+	UPROPERTY(EditInstanceOnly, Setter="SetOrientMesh", Getter="GetOrientMesh", Category="Layout")
 	bool bOrientMesh = false;
 
-	UPROPERTY(EditInstanceOnly, BlueprintReadWrite, Setter="SetRotation", Getter="GetRotation", Category="Layout")
+	UPROPERTY(EditInstanceOnly, Setter, Getter, Category="Layout")
 	FRotator Rotation = FRotator(0.f, 0.f, 0.f);
 
-	UPROPERTY(EditInstanceOnly, BlueprintReadWrite, Setter="SetScale", Getter="GetScale", Category="Layout", meta=(ClampMin="0", MotionDesignVectorWidget, AllowPreserveRatio="XYZ", Delta="0.01"))
+	UPROPERTY(EditInstanceOnly, Setter, Getter, Category="Layout", meta=(ClampMin="0", MotionDesignVectorWidget, AllowPreserveRatio="XYZ", Delta="0.01"))
 	FVector Scale = FVector(1.f, 1.f, 1.f);
 
 private:
diff --git a/Engine/Plugins/Experimental/ClonerEffector/Source/ClonerEffector/Public/Cloner/Layouts/CEClonerSphereUniformLayout.h b/Engine/Plugins/Experimental/ClonerEffector/Source/ClonerEffector/Public/Cloner/Layouts/CEClonerSphereUniformLayout.h
--- a/Engine/Plugins/Experimental/ClonerEffector/Source/ClonerEffector/Public/Cloner/Layouts/CEClonerSphereUniformLayout.h
+++ b/Engine/Plugins/Experimental/ClonerEffector/Source/ClonerEffector/Public/Cloner/Layouts/CEClonerSphereUniformLayout.h
@@ -86,22 +86,22 @@ protected:
 	virtual void OnLayoutParametersChanged(UCEClonerComponent* InComponent) override;
 	//~ End UCEClonerLayoutBase
 
-	UPROPERTY(EditInstanceOnly, BlueprintReadWrite, Setter="SetCount", Getter="GetCount", Category="Layout")
+	UPROPERTY(EditInstanceOnly, Setter, Getter, Category="Layout")
 	int32 Count = 30;
 
-	UPROPERTY(EditInstanceOnly, BlueprintReadWrite, Setter="SetRadius", Getter="GetRadius", Category="Layout")
+	UPROPERTY(EditInstanceOnly, Setter, Getter, Category="Layout")
 	float Radius = 200.f;
 
-	UPROPERTY(EditInstanceOnly, BlueprintReadWrite, Setter="SetRatio", Getter="GetRatio", Category="Layout", meta=(ClampMin="0", ClampMax="1"))
+	UPROPERTY(EditInstanceOnly, Setter, Getter, Category="Layout", meta=(ClampMin="0", ClampMax="1"))
 	float Ratio = 1.f;
 
-	UPROPERTY(EditInstanceOnly, BlueprintReadWrite, Setter="SetOrientMesh", Getter="GetOrientMesh", Category="Layout")
+	UPROPERTY(EditInstanceOnly, Setter="SetOrientMesh", Getter="GetOrientMesh", Category="Layout")
 	bool bOrientMesh = false;
 
-	UPROPERTY(EditInstanceOnly, BlueprintReadWrite, Setter="SetRotation", Getter="GetRotation", Category="Layout")
+	UPROPERTY(EditInstanceOnly, Setter, Getter, Category="Layout")
 	FRotator Rotation = FRotator(0.f, 0.f, 0.f);
 
-	UPROPERTY(EditInstanceOnly, BlueprintReadWrite, Setter="SetScale", Getter="GetScale", Category="Layout", meta=(ClampMin="0", MotionDesignVectorWidget, AllowPreserveRatio="XYZ", Delta="0.01"))
+	UPROPERTY(EditInstanceOnly, Setter, Getter, Category="Layout", meta=(ClampMin="0", MotionDesignVectorWidget, AllowPreserveRatio="XYZ", Delta="0.01"))
 	FVector Scale = FVector(1.f, 1.f, 1.f);
 
 private:
diff --git a/Engine/Plugins/Experimental/ClonerEffector/Source/ClonerEffector/Public/Cloner/Layouts/CEClonerSplineLayout.h b/Engine/Plugins/Experimental/ClonerEffector/Source/ClonerEffector/Public/Cloner/Layouts/CEClonerSplineLayout.h
--- a/Engine/Plugins/Experimental/ClonerEffector/Source/ClonerEffector/Public/Cloner/Layouts/CEClonerSplineLayout.h
+++ b/Engine/Plugins/Experimental/ClonerEffector/Source/ClonerEffector/Public/Cloner/Layouts/CEClonerSplineLayout.h
@@ -85,13 +85,13 @@ protected:
 	void OnSampleSplineTransformed(USceneComponent* InComponent, EUpdateTransformFlags InFlags, ETeleportType InType);
 	void OnSampleSplineRenderStateUpdated(UActorComponent& InComponent);
 
-	UPROPERTY(EditInstanceOnly, BlueprintReadWrite, Setter="SetCount", Getter="GetCount", Category="Layout")
+	UPROPERTY(EditInstanceOnly, Setter, Getter, Category="Layout")
 	int32 Count = 3 * 3;
 
-	UPROPERTY(EditInstanceOnly, BlueprintReadWrite, Setter="SetSplineActorWeak", Getter="GetSplineActorWeak", DisplayName="SampleActor", Category="Layout")
+	UPROPERTY(EditInstanceOnly, Setter, Getter, DisplayName="SampleActor", Category="Layout")
 	TWeakObjectPtr<AActor> SplineActorWeak;
 
-	UPROPERTY(EditInstanceOnly, BlueprintReadWrite, Setter="SetOrientMesh", Getter="GetOrientMesh", Category="Layout")
+	UPROPERTY(EditInstanceOnly, Setter="SetOrientMesh", Getter="GetOrientMesh", Category="Layout")
     bool bOrientMesh = false;
 
 private:
diff --git a/Engine/Plugins/Experimental/ClonerEffector/Source/ClonerEffector/Public/Effector/CEEffectorActor.h b/Engine/Plugins/Experimental/ClonerEffector/Source/ClonerEffector/Public/Effector/CEEffectorActor.h
--- a/Engine/Plugins/Experimental/ClonerEffector/Source/ClonerEffector/Public/Effector/CEEffectorActor.h
+++ b/Engine/Plugins/Experimental/ClonerEffector/Source/ClonerEffector/Public/Effector/CEEffectorActor.h
@@ -29,6 +29,7 @@ public:
 
 	ACEEffectorActor();
 
+	UFUNCTION(BlueprintPure, Category="Effector")
 	UCEEffectorComponent* GetEffectorComponent() const
 	{
 		return EffectorComponent;
diff --git a/Engine/Plugins/Experimental/ClonerEffector/Source/ClonerEffector/Public/Effector/CEEffectorComponent.h b/Engine/Plugins/Experimental/ClonerEffector/Source/ClonerEffector/Public/Effector/CEEffectorComponent.h
--- a/Engine/Plugins/Experimental/ClonerEffector/Source/ClonerEffector/Public/Effector/CEEffectorComponent.h
+++ b/Engine/Plugins/Experimental/ClonerEffector/Source/ClonerEffector/Public/Effector/CEEffectorComponent.h
@@ -74,6 +74,12 @@ public:
 		return TypeName;
 	}
 
+	UFUNCTION(BlueprintCallable, Category="Effector")
+	CLONEREFFECTOR_API void SetTypeClass(TSubclassOf<UCEEffectorTypeBase> InTypeClass);
+
+	UFUNCTION(BlueprintPure, Category="Effector")
+	CLONEREFFECTOR_API TSubclassOf<UCEEffectorTypeBase> GetTypeClass() const;
+
 	UFUNCTION(BlueprintCallable, Category="Effector")
 	CLONEREFFECTOR_API void SetModeName(FName InModeName);
 
@@ -83,6 +89,12 @@ public:
 		return ModeName;
 	}
 
+	UFUNCTION(BlueprintCallable, Category="Effector")
+	CLONEREFFECTOR_API void SetModeClass(TSubclassOf<UCEEffectorModeBase> InModeClass);
+
+	UFUNCTION(BlueprintPure, Category="Effector")
+	CLONEREFFECTOR_API TSubclassOf<UCEEffectorModeBase> GetModeClass() const;
+
 #if WITH_EDITOR
 	UFUNCTION(BlueprintCallable, Category="Effector")
 	CLONEREFFECTOR_API void SetVisualizerComponentVisible(bool bInVisible);
@@ -115,6 +127,7 @@ public:
 	FCEClonerEffectorChannelData& GetChannelData();
 
 	/** Get the effector channel identifier */
+	UFUNCTION(BlueprintPure, Category="Effector")
 	int32 GetChannelIdentifier() const;
 
 	void OnClonerLinked(UCEClonerEffectorExtension* InClonerExtension);
@@ -128,6 +141,7 @@ public:
 		return Cast<InTypeClass>(GetActiveType());
 	}
 
+	UFUNCTION(BlueprintPure, Category="Effector")
 	UCEEffectorTypeBase* GetActiveType() const
 	{
 		return ActiveType;
@@ -141,6 +155,7 @@ public:
 		return Cast<InModeClass>(GetActiveMode());
 	}
 
+	UFUNCTION(BlueprintPure, Category="Effector")
 	UCEEffectorModeBase* GetActiveMode() const
 	{
 		return ActiveMode;
@@ -151,6 +166,14 @@ public:
 		return ActiveExtensions;
 	}
 
+	UFUNCTION(BlueprintPure, Category="Effector")
+	void GetActiveExtensions(TArray<UCEEffectorExtensionBase*>& OutExtensions) const;
+
+	UFUNCTION(BlueprintCallable, Category="Effector")
+	UCEEffectorExtensionBase* GetExtension(TSubclassOf<UCEEffectorExtensionBase> InExtensionClass) const;
+
+	UCEEffectorExtensionBase* GetExtension(FName InExtensionName) const;
+
 	void RequestClonerUpdate(bool bInImmediate);
 
 protected:
@@ -175,19 +198,19 @@ protected:
 #endif
 
 	/** Is this effector enabled/disabled on linked cloners */
-	UPROPERTY(EditInstanceOnly, BlueprintReadWrite, Setter="SetEnabled", Getter="GetEnabled", Category="Effector")
+	UPROPERTY(EditInstanceOnly, Setter="SetEnabled", Getter="GetEnabled", Category="Effector")
 	bool bEnabled = true;
 
 	/** The ratio effect of the effector on clones */
-	UPROPERTY(EditInstanceOnly, BlueprintReadWrite, Setter, Getter, Category="Effector", meta=(ClampMin="0", ClampMax="1"))
+	UPROPERTY(EditInstanceOnly, Setter, Getter, Category="Effector", meta=(ClampMin="0", ClampMax="1"))
 	float Magnitude = 1.f;
 
 	/** Affected clones color passed over to material */
-	UPROPERTY(EditInstanceOnly, BlueprintReadWrite, Setter, Getter, Category="Effector")
+	UPROPERTY(EditInstanceOnly, Setter, Getter, Category="Effector")
 	FLinearColor Color = FLinearColor::Red;
 
 	/** Name of the shape type to use */
-	UPROPERTY(EditInstanceOnly, BlueprintReadWrite, Setter, Getter, Category="Shape", meta=(GetOptions="GetEffectorTypeNames"))
+	UPROPERTY(EditInstanceOnly, Setter, Getter, Category="Shape", meta=(GetOptions="GetEffectorTypeNames"))
 	FName TypeName = NAME_None;
 
 	/** Cached active type used for faster access */
@@ -195,7 +218,7 @@ protected:
 	TObjectPtr<UCEEffectorTypeBase> ActiveType;
 
 	/** Name of the shape type to use */
-	UPROPERTY(EditInstanceOnly, BlueprintReadWrite, Setter, Getter, Category="Mode", meta=(GetOptions="GetEffectorModeNames"))
+	UPROPERTY(EditInstanceOnly, Setter, Getter, Category="Mode", meta=(GetOptions="GetEffectorModeNames"))
 	FName ModeName = NAME_None;
 
 	/** Cached active mode used for faster access */
@@ -208,11 +231,11 @@ protected:
 
 #if WITH_EDITORONLY_DATA
 	/** Visibility of the components visualizer */
-	UPROPERTY(EditInstanceOnly, BlueprintReadWrite, AdvancedDisplay, Category="Effector")
+	UPROPERTY(EditInstanceOnly, AdvancedDisplay, Category="Effector")
 	bool bVisualizerComponentVisible = true;
 
 	/** Toggle the sprite to visualize and click on this effector */
-	UPROPERTY(EditInstanceOnly, BlueprintReadWrite, AdvancedDisplay, Category="Effector")
+	UPROPERTY(EditInstanceOnly, AdvancedDisplay, Category="Effector")
 	bool bVisualizerSpriteVisible = true;
 #endif
 
diff --git a/Engine/Plugins/Experimental/ClonerEffector/Source/ClonerEffector/Public/Effector/CEEffectorExtensionBase.h b/Engine/Plugins/Experimental/ClonerEffector/Source/ClonerEffector/Public/Effector/CEEffectorExtensionBase.h
--- a/Engine/Plugins/Experimental/ClonerEffector/Source/ClonerEffector/Public/Effector/CEEffectorExtensionBase.h
+++ b/Engine/Plugins/Experimental/ClonerEffector/Source/ClonerEffector/Public/Effector/CEEffectorExtensionBase.h
@@ -25,6 +25,7 @@ public:
 		: ExtensionName(InExtensionName)
 	{}
 
+	UFUNCTION(BlueprintPure, Category="Effector")
 	FName GetExtensionName() const
 	{
 		return ExtensionName;
@@ -46,6 +47,7 @@ public:
 	/** Disable this extension */
 	void DeactivateExtension();
 
+	UFUNCTION(BlueprintPure, Category="Effector")
 	bool IsExtensionActive() const
 	{
 		return bExtensionActive;
diff --git a/Engine/Plugins/Experimental/ClonerEffector/Source/ClonerEffector/Public/Effector/Extensions/CEEffectorForceExtension.h b/Engine/Plugins/Experimental/ClonerEffector/Source/ClonerEffector/Public/Effector/Extensions/CEEffectorForceExtension.h
--- a/Engine/Plugins/Experimental/ClonerEffector/Source/ClonerEffector/Public/Effector/Extensions/CEEffectorForceExtension.h
+++ b/Engine/Plugins/Experimental/ClonerEffector/Source/ClonerEffector/Public/Effector/Extensions/CEEffectorForceExtension.h
@@ -224,74 +224,74 @@ protected:
 	void OnForceOptionsChanged();
 
 	/** Forces global state */
-	UPROPERTY(EditInstanceOnly, BlueprintReadWrite, Setter="SetForcesEnabled", Getter="GetForcesEnabled", Category="Force")
+	UPROPERTY(EditInstanceOnly, Setter="SetForcesEnabled", Getter="GetForcesEnabled", Category="Force")
 	bool bForcesEnabled = false;
 
 	/** Enable orientation force to allow each clone instance to rotate around its pivot */
-	UPROPERTY(EditInstanceOnly, BlueprintReadWrite, Setter="SetOrientationForceEnabled", Getter="GetOrientationForceEnabled", Category="Force", meta=(EditCondition="bForcesEnabled", EditConditionHides))
+	UPROPERTY(EditInstanceOnly, Setter="SetOrientationForceEnabled", Getter="GetOrientationForceEnabled", Category="Force", meta=(EditCondition="bForcesEnabled", EditConditionHides))
 	bool bOrientationForceEnabled = false;
 
-	UPROPERTY(EditInstanceOnly, BlueprintReadWrite, Setter, Getter, Category="Force", meta=(Delta="0.0001", ClampMin="0", EditCondition="bForcesEnabled && bOrientationForceEnabled", EditConditionHides))
+	UPROPERTY(EditInstanceOnly, Setter, Getter, Category="Force", meta=(Delta="0.0001", ClampMin="0", EditCondition="bForcesEnabled && bOrientationForceEnabled", EditConditionHides))
 	float OrientationForceRate = 1.f;
 
-	UPROPERTY(EditInstanceOnly, BlueprintReadWrite, Setter, Getter, Category="Force", meta=(EditCondition="bForcesEnabled && bOrientationForceEnabled", EditConditionHides))
+	UPROPERTY(EditInstanceOnly, Setter, Getter, Category="Force", meta=(EditCondition="bForcesEnabled && bOrientationForceEnabled", EditConditionHides))
 	FVector OrientationForceMin = FVector(-0.1f);
 
-	UPROPERTY(EditInstanceOnly, BlueprintReadWrite, Setter, Getter, Category="Force", meta=(EditCondition="bForcesEnabled && bOrientationForceEnabled", EditConditionHides))
+	UPROPERTY(EditInstanceOnly, Setter, Getter, Category="Force", meta=(EditCondition="bForcesEnabled && bOrientationForceEnabled", EditConditionHides))
 	FVector OrientationForceMax = FVector(0.1f);
 
 	/** Enable vortex force to allow each clone instance to rotate around a specific axis on the cloner pivot */
-	UPROPERTY(EditInstanceOnly, BlueprintReadWrite, Setter="SetVortexForceEnabled", Getter="GetVortexForceEnabled", Category="Force", meta=(EditCondition="bForcesEnabled", EditConditionHides))
+	UPROPERTY(EditInstanceOnly, Setter="SetVortexForceEnabled", Getter="GetVortexForceEnabled", Category="Force", meta=(EditCondition="bForcesEnabled", EditConditionHides))
 	bool bVortexForceEnabled = false;
 
-	UPROPERTY(EditInstanceOnly, BlueprintReadWrite, Setter, Getter, Category="Force", meta=(EditCondition="bForcesEnabled && bVortexForceEnabled", EditConditionHides))
+	UPROPERTY(EditInstanceOnly, Setter, Getter, Category="Force", meta=(EditCondition="bForcesEnabled && bVortexForceEnabled", EditConditionHides))
 	float VortexForceAmount = 100.f;
 
-	UPROPERTY(EditInstanceOnly, BlueprintReadWrite, Setter, Getter, Category="Force", meta=(EditCondition="bForcesEnabled && bVortexForceEnabled", EditConditionHides))
+	UPROPERTY(EditInstanceOnly, Setter, Getter, Category="Force", meta=(EditCondition="bForcesEnabled && bVortexForceEnabled", EditConditionHides))
 	FVector VortexForceAxis = FVector::ZAxisVector;
 
 	/** Enable curl noise force to allow each clone instance to add random location variation */
-	UPROPERTY(EditInstanceOnly, BlueprintReadWrite, Setter="SetCurlNoiseForceEnabled", Getter="GetCurlNoiseForceEnabled", Category="Force", meta=(EditCondition="bForcesEnabled", EditConditionHides))
+	UPROPERTY(EditInstanceOnly, Setter="SetCurlNoiseForceEnabled", Getter="GetCurlNoiseForceEnabled", Category="Force", meta=(EditCondition="bForcesEnabled", EditConditionHides))
 	bool bCurlNoiseForceEnabled = false;
 
-	UPROPERTY(EditInstanceOnly, BlueprintReadWrite, Setter, Getter, Category="Force", meta=(EditCondition="bForcesEnabled && bCurlNoiseForceEnabled", EditConditionHides))
+	UPROPERTY(EditInstanceOnly, Setter, Getter, Category="Force", meta=(EditCondition="bForcesEnabled && bCurlNoiseForceEnabled", EditConditionHides))
 	float CurlNoiseForceStrength = 100.f;
 
-	UPROPERTY(EditInstanceOnly, BlueprintReadWrite, Setter, Getter, Category="Force", meta=(EditCondition="bForcesEnabled && bCurlNoiseForceEnabled", EditConditionHides))
+	UPROPERTY(EditInstanceOnly, Setter, Getter, Category="Force", meta=(EditCondition="bForcesEnabled && bCurlNoiseForceEnabled", EditConditionHides))
 	float CurlNoiseForceFrequency = 10.f;
 
 	/** Enable attraction force to allow each clone instances to gravitate toward a location */
-	UPROPERTY(EditInstanceOnly, BlueprintReadWrite, Setter="SetAttractionForceEnabled", Getter="GetAttractionForceEnabled", Category="Force", meta=(EditCondition="bForcesEnabled", EditConditionHides))
+	UPROPERTY(EditInstanceOnly, Setter="SetAttractionForceEnabled", Getter="GetAttractionForceEnabled", Category="Force", meta=(EditCondition="bForcesEnabled", EditConditionHides))
 	bool bAttractionForceEnabled = false;
 
-	UPROPERTY(EditInstanceOnly, BlueprintReadWrite, Setter, Getter, Category="Force", meta=(EditCondition="bForcesEnabled && bAttractionForceEnabled", EditConditionHides))
+	UPROPERTY(EditInstanceOnly, Setter, Getter, Category="Force", meta=(EditCondition="bForcesEnabled && bAttractionForceEnabled", EditConditionHides))
 	float AttractionForceStrength = 100.f;
 
-	UPROPERTY(EditInstanceOnly, BlueprintReadWrite, Setter, Getter, Category="Force", meta=(EditCondition="bForcesEnabled && bAttractionForceEnabled", EditConditionHides))
+	UPROPERTY(EditInstanceOnly, Setter, Getter, Category="Force", meta=(EditCondition="bForcesEnabled && bAttractionForceEnabled", EditConditionHides))
 	float AttractionForceFalloff = 0.1f;
 
 	/** Enable gravity force to pull particles based on an acceleration vector */
-	UPROPERTY(EditInstanceOnly, BlueprintReadWrite, Setter="SetGravityForceEnabled", Getter="GetGravityForceEnabled", Category="Force", meta=(EditCondition="bForcesEnabled", EditConditionHides))
+	UPROPERTY(EditInstanceOnly, Setter="SetGravityForceEnabled", Getter="GetGravityForceEnabled", Category="Force", meta=(EditCondition="bForcesEnabled", EditConditionHides))
 	bool bGravityForceEnabled = false;
 
-	UPROPERTY(EditInstanceOnly, BlueprintReadWrite, Setter, Getter, Category="Force", meta=(EditCondition="bForcesEnabled && bGravityForceEnabled", EditConditionHides))
+	UPROPERTY(EditInstanceOnly, Setter, Getter, Category="Force", meta=(EditCondition="bForcesEnabled && bGravityForceEnabled", EditConditionHides))
 	FVector GravityForceAcceleration = FVector(0, 0, -980.f);
 
 	/** Enable drag force to decrease particles velocity */
-	UPROPERTY(EditInstanceOnly, BlueprintReadWrite, Setter="SetDragForceEnabled", Getter="GetDragForceEnabled", Category="Force", meta=(EditCondition="bForcesEnabled", EditConditionHides))
+	UPROPERTY(EditInstanceOnly, Setter="SetDragForceEnabled", Getter="GetDragForceEnabled", Category="Force", meta=(EditCondition="bForcesEnabled", EditConditionHides))
 	bool bDragForceEnabled = false;
 
-	UPROPERTY(EditInstanceOnly, BlueprintReadWrite, Setter, Getter, Category="Force", meta=(EditCondition="bForcesEnabled && bDragForceEnabled", EditConditionHides))
+	UPROPERTY(EditInstanceOnly, Setter, Getter, Category="Force", meta=(EditCondition="bForcesEnabled && bDragForceEnabled", EditConditionHides))
 	float DragForceLinear = 0.f;
 
-	UPROPERTY(EditInstanceOnly, BlueprintReadWrite, Setter, Getter, Category="Force", meta=(EditCondition="bForcesEnabled && bDragForceEnabled", EditConditionHides))
+	UPROPERTY(EditInstanceOnly, Setter, Getter, Category="Force", meta=(EditCondition="bForcesEnabled && bDragForceEnabled", EditConditionHides))
 	float DragForceRotational = 0.f;
 
 	/** Enable vector random noise force to add variation to clones behavior */
-	UPROPERTY(EditInstanceOnly, BlueprintReadWrite, Setter="SetVectorNoiseForceEnabled", Getter="GetVectorNoiseForceEnabled", Category="Force", meta=(EditCondition="bForcesEnabled", EditConditionHides))
+	UPROPERTY(EditInstanceOnly, Setter="SetVectorNoiseForceEnabled", Getter="GetVectorNoiseForceEnabled", Category="Force", meta=(EditCondition="bForcesEnabled", EditConditionHides))
 	bool bVectorNoiseForceEnabled = false;
 
-	UPROPERTY(EditInstanceOnly, BlueprintReadWrite, Setter, Getter, Category="Force", meta=(EditCondition="bForcesEnabled && bVectorNoiseForceEnabled", EditConditionHides))
+	UPROPERTY(EditInstanceOnly, Setter, Getter, Category="Force", meta=(EditCondition="bForcesEnabled && bVectorNoiseForceEnabled", EditConditionHides))
 	float VectorNoiseForceAmount = 100.f;
 
 private:
diff --git a/Engine/Plugins/Experimental/ClonerEffector/Source/ClonerEffector/Public/Effector/Modes/CEEffectorNoiseMode.h b/Engine/Plugins/Experimental/ClonerEffector/Source/ClonerEffector/Public/Effector/Modes/CEEffectorNoiseMode.h
--- a/Engine/Plugins/Experimental/ClonerEffector/Source/ClonerEffector/Public/Effector/Modes/CEEffectorNoiseMode.h
+++ b/Engine/Plugins/Experimental/ClonerEffector/Source/ClonerEffector/Public/Effector/Modes/CEEffectorNoiseMode.h
@@ -76,23 +76,23 @@ protected:
 	//~ End UObject
 
 	/** Amplitude of the noise field for location */
-	UPROPERTY(EditInstanceOnly, BlueprintReadWrite, Setter, Getter, Category="Mode")
+	UPROPERTY(EditInstanceOnly, Setter, Getter, Category="Mode")
 	FVector LocationStrength = FVector::ZeroVector;
 
 	/** Amplitude of the noise field for rotation */
-	UPROPERTY(EditInstanceOnly, BlueprintReadWrite, Setter, Getter, Category="Mode")
+	UPROPERTY(EditInstanceOnly, Setter, Getter, Category="Mode")
 	FRotator RotationStrength = FRotator::ZeroRotator;
 
 	/** Amplitude of the noise field for scale */
-	UPROPERTY(EditInstanceOnly, BlueprintReadWrite, Setter, Getter, Category="Mode", meta=(ClampMin="0", MotionDesignVectorWidget, AllowPreserveRatio="XYZ", Delta="0.01"))
+	UPROPERTY(EditInstanceOnly, Setter, Getter, Category="Mode", meta=(ClampMin="0", MotionDesignVectorWidget, AllowPreserveRatio="XYZ", Delta="0.01"))
 	FVector ScaleStrength = FVector::OneVector;
 
 	/** Panning to offset the noise field sampling */
-	UPROPERTY(EditInstanceOnly, BlueprintReadWrite, Setter, Getter, Category="Mode")
+	UPROPERTY(EditInstanceOnly, Setter, Getter, Category="Mode")
 	FVector Pan = FVector::ZeroVector;
 
 	/** Intensity of the noise field */
-	UPROPERTY(EditInstanceOnly, BlueprintReadWrite, Setter, Getter, Category="Mode", meta=(ClampMin="0"))
+	UPROPERTY(EditInstanceOnly, Setter, Getter, Category="Mode", meta=(ClampMin="0"))
 	float Frequency = 0.5f;
 
 private:
diff --git a/Engine/Plugins/Experimental/ClonerEffector/Source/ClonerEffector/Public/Effector/Modes/CEEffectorOffsetMode.h b/Engine/Plugins/Experimental/ClonerEffector/Source/ClonerEffector/Public/Effector/Modes/CEEffectorOffsetMode.h
--- a/Engine/Plugins/Experimental/ClonerEffector/Source/ClonerEffector/Public/Effector/Modes/CEEffectorOffsetMode.h
+++ b/Engine/Plugins/Experimental/ClonerEffector/Source/ClonerEffector/Public/Effector/Modes/CEEffectorOffsetMode.h
@@ -58,15 +58,15 @@ protected:
 	//~ End UObject
 
 	/** Offset applied on affected clones */
-	UPROPERTY(EditInstanceOnly, BlueprintReadWrite, Setter, Getter, Category="Mode")
+	UPROPERTY(EditInstanceOnly, Setter, Getter, Category="Mode")
 	FVector Offset = FVector::ZeroVector;
 
 	/** Rotation applied on affected clones */
-	UPROPERTY(EditInstanceOnly, BlueprintReadWrite, Setter, Getter, Category="Mode")
+	UPROPERTY(EditInstanceOnly, Setter, Getter, Category="Mode")
 	FRotator Rotation = FRotator::ZeroRotator;
 
 	/** Scale applied on affected clones */
-	UPROPERTY(EditInstanceOnly, BlueprintReadWrite, Setter, Getter, Category="Mode", meta=(ClampMin="0", MotionDesignVectorWidget, AllowPreserveRatio="XYZ", Delta="0.01"))
+	UPROPERTY(EditInstanceOnly, Setter, Getter, Category="Mode", meta=(ClampMin="0", MotionDesignVectorWidget, AllowPreserveRatio="XYZ", Delta="0.01"))
 	FVector Scale = FVector::OneVector;
 
 private:
diff --git a/Engine/Plugins/Experimental/ClonerEffector/Source/ClonerEffector/Public/Effector/Modes/CEEffectorPushMode.h b/Engine/Plugins/Experimental/ClonerEffector/Source/ClonerEffector/Public/Effector/Modes/CEEffectorPushMode.h
--- a/Engine/Plugins/Experimental/ClonerEffector/Source/ClonerEffector/Public/Effector/Modes/CEEffectorPushMode.h
+++ b/Engine/Plugins/Experimental/ClonerEffector/Source/ClonerEffector/Public/Effector/Modes/CEEffectorPushMode.h
@@ -49,11 +49,11 @@ protected:
 	//~ End UObject
 
 	/** Strength of the push effect */
-	UPROPERTY(EditInstanceOnly, BlueprintReadWrite, Setter, Getter, Category="Mode", meta=(MotionDesignVectorWidget, AllowPreserveRatio="XYZ"))
+	UPROPERTY(EditInstanceOnly, Setter, Getter, Category="Mode", meta=(MotionDesignVectorWidget, AllowPreserveRatio="XYZ"))
 	FVector PushStrength = FVector(100.f);
 
 	/** Relative direction computed for the push effect on each clone */
-	UPROPERTY(EditInstanceOnly, BlueprintReadWrite, Setter, Getter, Category="Mode")
+	UPROPERTY(EditInstanceOnly, Setter, Getter, Category="Mode")
 	ECEClonerEffectorPushDirection PushDirection = ECEClonerEffectorPushDirection::Forward;
 
 private:
diff --git a/Engine/Plugins/Experimental/ClonerEffector/Source/ClonerEffector/Public/Effector/Modes/CEEffectorStepMode.h b/Engine/Plugins/Experimental/ClonerEffector/Source/ClonerEffector/Public/Effector/Modes/CEEffectorStepMode.h
--- a/Engine/Plugins/Experimental/ClonerEffector/Source/ClonerEffector/Public/Effector/Modes/CEEffectorStepMode.h
+++ b/Engine/Plugins/Experimental/ClonerEffector/Source/ClonerEffector/Public/Effector/Modes/CEEffectorStepMode.h
@@ -58,15 +58,15 @@ protected:
 	//~ End UCEEffectorNoiseMode
 
 	/** Interpolates from 0 to this position offset based on the particle index and particle count */
-	UPROPERTY(EditInstanceOnly, BlueprintReadWrite, Setter, Getter, DisplayName="Position", Category="Mode")
+	UPROPERTY(EditInstanceOnly, Setter, Getter, DisplayName="Position", Category="Mode")
 	FVector StepPosition = FVector::ZeroVector;
 
 	/** Interpolates from 0 to this rotation based on the particle index and particle count */
-	UPROPERTY(EditInstanceOnly, BlueprintReadWrite, Setter, Getter, DisplayName="Rotation", Category="Mode")
+	UPROPERTY(EditInstanceOnly, Setter, Getter, DisplayName="Rotation", Category="Mode")
 	FRotator StepRotation = FRotator::ZeroRotator;
 
 	/** Interpolates from 1 to this scale based on the particle index and particle count */
-	UPROPERTY(EditInstanceOnly, BlueprintReadWrite, Setter, Getter, DisplayName="Scale", Category="Mode", meta=(ClampMin="0", MotionDesignVectorWidget, AllowPreserveRatio="XYZ"))
+	UPROPERTY(EditInstanceOnly, Setter, Getter, DisplayName="Scale", Category="Mode", meta=(ClampMin="0", MotionDesignVectorWidget, AllowPreserveRatio="XYZ"))
 	FVector StepScale = FVector::OneVector;
 
 private:
diff --git a/Engine/Plugins/Experimental/ClonerEffector/Source/ClonerEffector/Public/Effector/Modes/CEEffectorTargetMode.h b/Engine/Plugins/Experimental/ClonerEffector/Source/ClonerEffector/Public/Effector/Modes/CEEffectorTargetMode.h
--- a/Engine/Plugins/Experimental/ClonerEffector/Source/ClonerEffector/Public/Effector/Modes/CEEffectorTargetMode.h
+++ b/Engine/Plugins/Experimental/ClonerEffector/Source/ClonerEffector/Public/Effector/Modes/CEEffectorTargetMode.h
@@ -50,7 +50,7 @@ protected:
 	//~ End UObject
 
 	/** The actor to track when mode is set to target */
-	UPROPERTY(EditInstanceOnly, BlueprintReadWrite, Setter, Getter, Category="Mode", meta=(DisplayName="TargetActor"))
+	UPROPERTY(EditInstanceOnly, Setter, Getter, Category="Mode", meta=(DisplayName="TargetActor"))
 	TWeakObjectPtr<AActor> TargetActorWeak = nullptr;
 
 	UPROPERTY()
diff --git a/Engine/Plugins/Experimental/ClonerEffector/Source/ClonerEffector/Public/Effector/Types/CEEffectorBoundType.h b/Engine/Plugins/Experimental/ClonerEffector/Source/ClonerEffector/Public/Effector/Types/CEEffectorBoundType.h
--- a/Engine/Plugins/Experimental/ClonerEffector/Source/ClonerEffector/Public/Effector/Types/CEEffectorBoundType.h
+++ b/Engine/Plugins/Experimental/ClonerEffector/Source/ClonerEffector/Public/Effector/Types/CEEffectorBoundType.h
@@ -53,11 +53,11 @@ protected:
 	//~ End UObject
 
 	/** Invert the type effect, instead of affecting the inside of a zone, will affect the outside */
-	UPROPERTY(EditInstanceOnly, BlueprintReadWrite, Setter="SetInvertType", Getter="GetInvertType", Category="Shape")
+	UPROPERTY(EditInstanceOnly, Setter="SetInvertType", Getter="GetInvertType", Category="Shape")
 	bool bInvertType = false;
 
 	/** Weight easing function applied to lerp transforms */
-	UPROPERTY(EditInstanceOnly, BlueprintReadWrite, Setter, Getter, Category="Shape")
+	UPROPERTY(EditInstanceOnly, Setter, Getter, Category="Shape")
 	ECEClonerEasing Easing = ECEClonerEasing::Linear;
 
 private:
diff --git a/Engine/Plugins/Experimental/ClonerEffector/Source/ClonerEffector/Public/Effector/Types/CEEffectorBoxType.h b/Engine/Plugins/Experimental/ClonerEffector/Source/ClonerEffector/Public/Effector/Types/CEEffectorBoxType.h
--- a/Engine/Plugins/Experimental/ClonerEffector/Source/ClonerEffector/Public/Effector/Types/CEEffectorBoxType.h
+++ b/Engine/Plugins/Experimental/ClonerEffector/Source/ClonerEffector/Public/Effector/Types/CEEffectorBoxType.h
@@ -54,11 +54,11 @@ protected:
 	//~ End UObject
 
 	/** Inner extent of box, all clones inside will be affected with a maximum weight */
-	UPROPERTY(EditInstanceOnly, BlueprintReadWrite, Setter, Getter, Category="Shape", meta=(ClampMin="0", MotionDesignVectorWidget, AllowPreserveRatio="XYZ"))
+	UPROPERTY(EditInstanceOnly, Setter, Getter, Category="Shape", meta=(ClampMin="0", MotionDesignVectorWidget, AllowPreserveRatio="XYZ"))
 	FVector InnerExtent = FVector(50.f);
 
 	/** Outer extent of box, all clones outside will not be affected */
-	UPROPERTY(EditInstanceOnly, BlueprintReadWrite, Setter, Getter, Category="Shape", meta=(ClampMin="0", MotionDesignVectorWidget, AllowPreserveRatio="XYZ"))
+	UPROPERTY(EditInstanceOnly, Setter, Getter, Category="Shape", meta=(ClampMin="0", MotionDesignVectorWidget, AllowPreserveRatio="XYZ"))
 	FVector OuterExtent = FVector(200.f);
 
 private:
diff --git a/Engine/Plugins/Experimental/ClonerEffector/Source/ClonerEffector/Public/Effector/Types/CEEffectorPlaneType.h b/Engine/Plugins/Experimental/ClonerEffector/Source/ClonerEffector/Public/Effector/Types/CEEffectorPlaneType.h
--- a/Engine/Plugins/Experimental/ClonerEffector/Source/ClonerEffector/Public/Effector/Types/CEEffectorPlaneType.h
+++ b/Engine/Plugins/Experimental/ClonerEffector/Source/ClonerEffector/Public/Effector/Types/CEEffectorPlaneType.h
@@ -45,7 +45,7 @@ protected:
 	//~ End UObject
 
 	/** Plane spacing, everything inside this zone will be affected */
-	UPROPERTY(EditInstanceOnly, BlueprintReadWrite, Setter, Getter, Category="Shape", meta=(ClampMin="0"))
+	UPROPERTY(EditInstanceOnly, Setter, Getter, Category="Shape", meta=(ClampMin="0"))
 	float PlaneSpacing = 200.f;
 
 private:
diff --git a/Engine/Plugins/Experimental/ClonerEffector/Source/ClonerEffector/Public/Effector/Types/CEEffectorRadialType.h b/Engine/Plugins/Experimental/ClonerEffector/Source/ClonerEffector/Public/Effector/Types/CEEffectorRadialType.h
--- a/Engine/Plugins/Experimental/ClonerEffector/Source/ClonerEffector/Public/Effector/Types/CEEffectorRadialType.h
+++ b/Engine/Plugins/Experimental/ClonerEffector/Source/ClonerEffector/Public/Effector/Types/CEEffectorRadialType.h
@@ -63,15 +63,15 @@ protected:
 	//~ End UObject
 
 	/** Radial angle in degree, everything within the angle will be affected */
-	UPROPERTY(EditInstanceOnly, BlueprintReadWrite, Setter, Getter, Category="Shape", meta=(ClampMin="0", ClampMax="360"))
+	UPROPERTY(EditInstanceOnly, Setter, Getter, Category="Shape", meta=(ClampMin="0", ClampMax="360"))
 	float RadialAngle = 180.f;
 
 	/** Minimum radius for the radial effect to be applied on clones, below clones will not be affected */
-	UPROPERTY(EditInstanceOnly, BlueprintReadWrite, Setter, Getter, Category="Shape", meta=(ClampMin="0"))
+	UPROPERTY(EditInstanceOnly, Setter, Getter, Category="Shape", meta=(ClampMin="0"))
 	float RadialMinRadius = 0.f;
 
 	/** Maximum radius for the radial effect to be applied on clones, above clones will not be affected */
-	UPROPERTY(EditInstanceOnly, BlueprintReadWrite, Setter, Getter, Category="Shape", meta=(ClampMin="0"))
+	UPROPERTY(EditInstanceOnly, Setter, Getter, Category="Shape", meta=(ClampMin="0"))
 	float RadialMaxRadius = 1000.f;
 
 private:
diff --git a/Engine/Plugins/Experimental/ClonerEffector/Source/ClonerEffector/Public/Effector/Types/CEEffectorSphereType.h b/Engine/Plugins/Experimental/ClonerEffector/Source/ClonerEffector/Public/Effector/Types/CEEffectorSphereType.h
--- a/Engine/Plugins/Experimental/ClonerEffector/Source/ClonerEffector/Public/Effector/Types/CEEffectorSphereType.h
+++ b/Engine/Plugins/Experimental/ClonerEffector/Source/ClonerEffector/Public/Effector/Types/CEEffectorSphereType.h
@@ -54,11 +54,11 @@ protected:
 	//~ End UObject
 
 	/** Inner radius of sphere, all clones inside will be affected with a maximum weight */
-	UPROPERTY(EditInstanceOnly, BlueprintReadWrite, Setter, Getter, Category="Shape", meta=(ClampMin="0"))
+	UPROPERTY(EditInstanceOnly, Setter, Getter, Category="Shape", meta=(ClampMin="0"))
 	float InnerRadius = 50.f;
 
 	/** Outer radius of sphere, all clones outside will not be affected */
-	UPROPERTY(EditInstanceOnly, BlueprintReadWrite, Setter, Getter, Category="Shape", meta=(ClampMin="0"))
+	UPROPERTY(EditInstanceOnly, Setter, Getter, Category="Shape", meta=(ClampMin="0"))
 	float OuterRadius = 200.f;
 
 private:
diff --git a/Engine/Plugins/Experimental/ClonerEffector/Source/ClonerEffector/Public/Effector/Types/CEEffectorTorusType.h b/Engine/Plugins/Experimental/ClonerEffector/Source/ClonerEffector/Public/Effector/Types/CEEffectorTorusType.h
--- a/Engine/Plugins/Experimental/ClonerEffector/Source/ClonerEffector/Public/Effector/Types/CEEffectorTorusType.h
+++ b/Engine/Plugins/Experimental/ClonerEffector/Source/ClonerEffector/Public/Effector/Types/CEEffectorTorusType.h
@@ -63,15 +63,15 @@ protected:
 	//~ End UObject
 
 	/** Main torus radius from center to the edge where inner and outer tube will be revolved */
-	UPROPERTY(EditInstanceOnly, BlueprintReadWrite, Setter, Getter, Category="Shape", meta=(ClampMin="0"))
+	UPROPERTY(EditInstanceOnly, Setter, Getter, Category="Shape", meta=(ClampMin="0"))
 	float TorusRadius = 250.f;
 
 	/** Minimum revolved radius for the torus effect, clones contained inside will be affected with a maximum weight */
-	UPROPERTY(EditInstanceOnly, BlueprintReadWrite, Setter, Getter, Category="Shape", meta=(ClampMin="0"))
+	UPROPERTY(EditInstanceOnly, Setter, Getter, Category="Shape", meta=(ClampMin="0"))
 	float TorusInnerRadius = 50.f;
 
 	/** Maximum revolved radius for the torus effect, clones outside of it will not be affected */
-	UPROPERTY(EditInstanceOnly, BlueprintReadWrite, Setter, Getter, Category="Shape", meta=(ClampMin="0"))
+	UPROPERTY(EditInstanceOnly, Setter, Getter, Category="Shape", meta=(ClampMin="0"))
 	float TorusOuterRadius = 200.f;
 
 private:
diff --git a/Engine/Plugins/Experimental/ClonerEffector/Source/ClonerEffector/Public/Subsystems/CEClonerSubsystem.h b/Engine/Plugins/Experimental/ClonerEffector/Source/ClonerEffector/Public/Subsystems/CEClonerSubsystem.h
--- a/Engine/Plugins/Experimental/ClonerEffector/Source/ClonerEffector/Public/Subsystems/CEClonerSubsystem.h
+++ b/Engine/Plugins/Experimental/ClonerEffector/Source/ClonerEffector/Public/Subsystems/CEClonerSubsystem.h
@@ -15,13 +15,16 @@ class UCEClonerSubsystem : public UEngineSubsystem
 {
 	GENERATED_BODY()
 
+	DECLARE_MULTICAST_DELEGATE(FOnSubsystemInitialized)
+	DECLARE_MULTICAST_DELEGATE_ThreeParams(FOnClonerSetEnabled, const UWorld* /** InWorld */, bool /** bInEnabled */, bool /** bInTransact */)
+
 public:
-	static TMulticastDelegateRegistration<void()>& OnSubsystemInitialized()
+	static FOnSubsystemInitialized::RegistrationType& OnSubsystemInitialized()
 	{
 		return OnSubsystemInitializedDelegate;
 	}
 
-	static TMulticastDelegateRegistration<void(const UWorld*, bool, bool)>& OnClonerSetEnabled()
+	static FOnClonerSetEnabled::RegistrationType& OnClonerSetEnabled()
 	{
 		return OnClonerSetEnabledDelegate;
 	}
@@ -36,7 +39,10 @@ public:
 	CLONEREFFECTOR_API bool IsLayoutClassRegistered(UClass* InClonerLayoutClass);
 
 	/** Get available cloner layout names to use in dropdown */
-	TArray<FName> GetLayoutNames() const;
+	TSet<FName> GetLayoutNames() const;
+
+	/** Get available cloner layout classes */
+	TSet<TSubclassOf<UCEClonerLayoutBase>> GetLayoutClasses() const;
 
 	/** Based on a layout class, find layout name */
 	FName FindLayoutName(TSubclassOf<UCEClonerLayoutBase> InLayoutClass) const;
@@ -87,11 +93,9 @@ public:
 	CLONEREFFECTOR_API AActor* CreateClonerWithActors(UWorld* InWorld, const TSet<AActor*>& InActors, bool bInShouldTransact);
 
 protected:
-	DECLARE_MULTICAST_DELEGATE(FOnSubsystemInitialized)
 	CLONEREFFECTOR_API static FOnSubsystemInitialized OnSubsystemInitializedDelegate;
 
 	/** Delegate to change state of cloners in a world */
-	DECLARE_MULTICAST_DELEGATE_ThreeParams(FOnClonerSetEnabled, const UWorld* /** InWorld */, bool /** bInEnabled */, bool /** bInTransact */)
 	static FOnClonerSetEnabled OnClonerSetEnabledDelegate;
 
 	//~ Begin USubsystem
diff --git a/Engine/Plugins/Experimental/ClonerEffector/Source/ClonerEffector/Public/Subsystems/CEEffectorSubsystem.h b/Engine/Plugins/Experimental/ClonerEffector/Source/ClonerEffector/Public/Subsystems/CEEffectorSubsystem.h
--- a/Engine/Plugins/Experimental/ClonerEffector/Source/ClonerEffector/Public/Subsystems/CEEffectorSubsystem.h
+++ b/Engine/Plugins/Experimental/ClonerEffector/Source/ClonerEffector/Public/Subsystems/CEEffectorSubsystem.h
@@ -16,18 +16,22 @@ class UCEEffectorSubsystem : public UEngineSubsystem, public FTickableGameObject
 {
 	GENERATED_BODY()
 
+	DECLARE_MULTICAST_DELEGATE(FOnSubsystemInitialized)
+	DECLARE_MULTICAST_DELEGATE_ThreeParams(FOnEffectorIdentifierChanged, UCEEffectorComponent* /** InEffector */, int32 /** OldIdentifier */, int32 /** NewIdentifier */)
+	DECLARE_MULTICAST_DELEGATE_ThreeParams(FOnEffectorSetEnabled, const UWorld* /** InWorld */, bool /** bInEnabled */, bool /** bInTransact */)
+
 public:
-	static TMulticastDelegateRegistration<void()>& OnSubsystemInitialized()
+	static FOnSubsystemInitialized::RegistrationType& OnSubsystemInitialized()
 	{
 		return OnSubsystemInitializedDelegate;
 	}
 
-	static TMulticastDelegateRegistration<void(UCEEffectorComponent*, int32, int32)>& OnEffectorIdentifierChanged()
+	static FOnEffectorIdentifierChanged::RegistrationType& OnEffectorIdentifierChanged()
 	{
 		return OnEffectorIdentifierChangedDelegate;
 	}
 
-	static TMulticastDelegateRegistration<void(const UWorld*, bool, bool)>& OnEffectorSetEnabled()
+	static FOnEffectorSetEnabled::RegistrationType& OnEffectorSetEnabled()
 	{
 		return OnEffectorSetEnabledDelegate;
 	}
@@ -51,12 +55,20 @@ public:
 	bool IsExtensionClassRegistered(UClass* InClass) const;
 
 	template<typename InExtensionClass>
-	TArray<FName> GetExtensionNames() const
+	TSet<FName> GetExtensionNames() const
 	{
 		return GetExtensionNames(InExtensionClass::StaticClass());
 	}
 
-	TArray<FName> GetExtensionNames(TSubclassOf<UCEEffectorExtensionBase> InExtensionClass) const;
+	TSet<FName> GetExtensionNames(TSubclassOf<UCEEffectorExtensionBase> InExtensionClass) const;
+
+	template<typename InExtensionClass>
+	TSet<TSubclassOf<UCEEffectorExtensionBase>> GetExtensionClasses() const
+	{
+		return GetExtensionClasses(InExtensionClass::StaticClass());
+	}
+
+	TSet<TSubclassOf<UCEEffectorExtensionBase>> GetExtensionClasses(TSubclassOf<UCEEffectorExtensionBase> InExtensionClass) const;
 
 	/** Based on a extension class, find extension name */
 	FName FindExtensionName(TSubclassOf<UCEEffectorExtensionBase> InClass) const;
@@ -74,15 +86,12 @@ protected:
 	static constexpr TCHAR DataChannelAssetPath[] = TEXT("/Script/Niagara.NiagaraDataChannelAsset'/ClonerEffector/Channels/NDC_Effector.NDC_Effector'");
 
 	/** Broadcasted when this subsystem is initialized */
-	DECLARE_MULTICAST_DELEGATE(FOnSubsystemInitialized)
 	static FOnSubsystemInitialized OnSubsystemInitializedDelegate;
 
 	/** Broadcasted when this effector identifier changed to update linked cloners */
-	DECLARE_MULTICAST_DELEGATE_ThreeParams(FOnEffectorIdentifierChanged, UCEEffectorComponent* /** InEffector */, int32 /** OldIdentifier */, int32 /** NewIdentifier */)
 	static FOnEffectorIdentifierChanged OnEffectorIdentifierChangedDelegate;
 
 	/** Delegate to change state of effectors in a world */
-	DECLARE_MULTICAST_DELEGATE_ThreeParams(FOnEffectorSetEnabled, const UWorld* /** InWorld */, bool /** bInEnabled */, bool /** bInTransact */)
 	static FOnEffectorSetEnabled OnEffectorSetEnabledDelegate;
 
 	//~ Begin USubsystem
diff --git a/Engine/Plugins/Experimental/ClonerEffector/Source/ClonerEffector/Public/Utilities/CEClonerLibrary.h b/Engine/Plugins/Experimental/ClonerEffector/Source/ClonerEffector/Public/Utilities/CEClonerLibrary.h
new file mode 100644
--- /dev/null
+++ b/Engine/Plugins/Experimental/ClonerEffector/Source/ClonerEffector/Public/Utilities/CEClonerLibrary.h
@@ -0,0 +1,34 @@
+// Copyright Epic Games, Inc. All Rights Reserved.
+
+#pragma once
+
+#include "Containers/ContainersFwd.h"
+#include "Kismet/BlueprintFunctionLibrary.h"
+#include "Templates/SubclassOf.h"
+#include "CEClonerLibrary.generated.h"
+
+template <typename T> class TSubclassOf;
+class UCEClonerExtensionBase;
+class UCEClonerLayoutBase;
+
+/** Blueprint operations for cloner */
+UCLASS(MinimalAPI, DisplayName="Motion Design Cloner Library")
+class UCEClonerLibrary : public UBlueprintFunctionLibrary
+{
+	GENERATED_BODY()
+
+public:
+	/**
+	 * Retrieves all layout classes available for a cloner
+	 * @param OutLayoutClasses [Out] Layout classes available
+	 */
+	UFUNCTION(BlueprintCallable, Category="Motion Design|Cloner|Utility")
+	static CLONEREFFECTOR_API void GetClonerLayoutClasses(TSet<TSubclassOf<UCEClonerLayoutBase>>& OutLayoutClasses);
+
+	/**
+	 * Retrieves all extension classes available for a cloner
+	 * @param OutExtensionClasses [Out] Extension classes available
+	 */
+	UFUNCTION(BlueprintCallable, Category="Motion Design|Cloner|Utility")
+	static CLONEREFFECTOR_API void GetClonerExtensionClasses(TSet<TSubclassOf<UCEClonerExtensionBase>>& OutExtensionClasses);
+};
\ No newline at end of file
diff --git a/Engine/Plugins/Experimental/ClonerEffector/Source/ClonerEffector/Public/Utilities/CEEffectorLibrary.h b/Engine/Plugins/Experimental/ClonerEffector/Source/ClonerEffector/Public/Utilities/CEEffectorLibrary.h
new file mode 100644
--- /dev/null
+++ b/Engine/Plugins/Experimental/ClonerEffector/Source/ClonerEffector/Public/Utilities/CEEffectorLibrary.h
@@ -0,0 +1,41 @@
+// Copyright Epic Games, Inc. All Rights Reserved.
+
+#pragma once
+
+#include "Containers/ContainersFwd.h"
+#include "Kismet/BlueprintFunctionLibrary.h"
+#include "Templates/SubclassOf.h"
+#include "CEEffectorLibrary.generated.h"
+
+class UCEEffectorExtensionBase;
+class UCEEffectorTypeBase;
+class UCEEffectorModeBase;
+
+/** Blueprint operations for effector */
+UCLASS(MinimalAPI, DisplayName="Motion Design Effector Library")
+class UCEEffectorLibrary : public UBlueprintFunctionLibrary
+{
+	GENERATED_BODY()
+
+public:
+	/**
+	 * Retrieves all mode classes available for an effector
+	 * @param OutModeClasses [Out] Available mode classes
+	 */
+	UFUNCTION(BlueprintCallable, Category="Motion Design|Effector|Utility")
+	static CLONEREFFECTOR_API void GetEffectorModeClasses(TSet<TSubclassOf<UCEEffectorModeBase>>& OutModeClasses);
+
+	/**
+	 * Retrieves all type classes available for an effector
+	 * @param OutTypeClasses [Out] Available type classes
+	 */
+	UFUNCTION(BlueprintCallable, Category="Motion Design|Effector|Utility")
+	static CLONEREFFECTOR_API void GetEffectorTypeClasses(TSet<TSubclassOf<UCEEffectorTypeBase>>& OutTypeClasses);
+
+	/**
+	 * Retrieves all extension classes available for an effector
+	 * @param OutExtensionClasses [Out] Available extension classes
+	 */
+	UFUNCTION(BlueprintCallable, Category="Motion Design|Effector|Utility")
+	static CLONEREFFECTOR_API void GetEffectorExtensionClasses(TSet<TSubclassOf<UCEEffectorExtensionBase>>& OutExtensionClasses);
+};
\ No newline at end of file
