MotionDesign : ClonerEffector
- Adding missing blueprint function with async support to set and wait until layout is ready

diff --git a/Engine/Plugins/Experimental/ClonerEffector/Source/ClonerEffector/Private/Subsystems/CEClonerSubsystem.cpp b/Engine/Plugins/Experimental/ClonerEffector/Source/ClonerEffector/Private/Subsystems/CEClonerSubsystem.cpp
--- a/Engine/Plugins/Experimental/ClonerEffector/Source/ClonerEffector/Private/Subsystems/CEClonerSubsystem.cpp
+++ b/Engine/Plugins/Experimental/ClonerEffector/Source/ClonerEffector/Private/Subsystems/CEClonerSubsystem.cpp
@@ -461,6 +461,16 @@ FName UCEClonerSubsystem::FindLayoutName(TSubclassOf<UCEClonerLayoutBase> InLayo
 	return NAME_None;
 }
 
+TSubclassOf<UCEClonerLayoutBase> UCEClonerSubsystem::FindLayoutClass(FName InLayoutName) const
+{
+	if (const TSubclassOf<UCEClonerLayoutBase>* Value = LayoutClasses.Find(InLayoutName))
+	{
+		return *Value;
+	}
+
+	return TSubclassOf<UCEClonerLayoutBase>();
+}
+
 UCEClonerLayoutBase* UCEClonerSubsystem::CreateNewLayout(FName InLayoutName, UCEClonerComponent* InCloner)
 {
 	if (!IsValid(InCloner))
diff --git a/Engine/Plugins/Experimental/ClonerEffector/Source/ClonerEffector/Private/Utilities/CEClonerLayoutLatentAction.cpp b/Engine/Plugins/Experimental/ClonerEffector/Source/ClonerEffector/Private/Utilities/CEClonerLayoutLatentAction.cpp
new file mode 100644
--- /dev/null
+++ b/Engine/Plugins/Experimental/ClonerEffector/Source/ClonerEffector/Private/Utilities/CEClonerLayoutLatentAction.cpp
@@ -0,0 +1,79 @@
+// Copyright Epic Games, Inc. All Rights Reserved.
+
+#include "Utilities/CEClonerLayoutLatentAction.h"
+
+#include "Cloner/CEClonerComponent.h"
+#include "Cloner/Layouts/CEClonerLayoutBase.h"
+#include "Engine/LatentActionManager.h"
+#include "HAL/Platform.h"
+
+FCEClonerLayoutLatentAction::FCEClonerLayoutLatentAction(const FLatentActionInfo& InLatentInfo, UCEClonerComponent* InCloner, TSubclassOf<UCEClonerLayoutBase> InLayoutClass, bool& bOutSuccess, UCEClonerLayoutBase*& OutLayout)
+	: ExecutionFunction(InLatentInfo.ExecutionFunction)
+	, OutputLink(InLatentInfo.Linkage)
+	, CallbackTarget(InLatentInfo.CallbackTarget)
+	, bLayoutSet(false)
+	, ClonerWeak(InCloner)
+	, LayoutClass(InLayoutClass)
+	, bSuccess(bOutSuccess)
+	, Layout(OutLayout)
+{
+}
+
+FCEClonerLayoutLatentAction::~FCEClonerLayoutLatentAction()
+{
+	UCEClonerComponent::OnClonerLayoutLoaded().RemoveAll(this);
+}
+
+#if WITH_EDITOR
+FString FCEClonerLayoutLatentAction::GetDescription() const
+{
+	const UCEClonerComponent* Cloner = ClonerWeak.Get();
+	return FString::Printf(TEXT("Cloner %s layout %s loaded : %i"), Cloner ? *Cloner->GetName() : TEXT("Invalid"), *LayoutClass->GetName(), bSuccess);
+}
+#endif
+
+void FCEClonerLayoutLatentAction::OnClonerLayoutLoaded(UCEClonerComponent* InClonerComponent, UCEClonerLayoutBase* InClonerLayoutBase) const
+{
+	if (InClonerComponent == ClonerWeak.Get())
+	{
+		if (InClonerLayoutBase && InClonerLayoutBase->GetClass() == LayoutClass)
+		{
+			Layout = InClonerLayoutBase;
+			bSuccess = true;
+		}
+	}
+}
+
+void FCEClonerLayoutLatentAction::UpdateOperation(FLatentResponse& OutResponse)
+{
+	UCEClonerComponent* Cloner = ClonerWeak.Get();
+
+	if (Cloner && !bLayoutSet)
+	{
+		bLayoutSet = true;
+
+		if (Cloner->GetLayoutClass() != LayoutClass)
+		{
+			UCEClonerComponent::OnClonerLayoutLoaded().AddRaw(this, &FCEClonerLayoutLatentAction::OnClonerLayoutLoaded);
+
+			if (Cloner)
+			{
+				Cloner->SetLayoutClass(LayoutClass);
+			}
+		}
+		else
+		{
+			OnClonerLayoutLoaded(Cloner, Cloner->GetActiveLayout());
+		}
+	}
+
+	bool bDone = bSuccess;
+
+	if (!Cloner)
+	{
+		bSuccess = false;
+		bDone = true;
+	}
+
+	OutResponse.FinishAndTriggerIf(bDone, ExecutionFunction, OutputLink, CallbackTarget);
+}
\ No newline at end of file
diff --git a/Engine/Plugins/Experimental/ClonerEffector/Source/ClonerEffector/Private/Utilities/CEClonerLayoutLatentAction.h b/Engine/Plugins/Experimental/ClonerEffector/Source/ClonerEffector/Private/Utilities/CEClonerLayoutLatentAction.h
new file mode 100644
--- /dev/null
+++ b/Engine/Plugins/Experimental/ClonerEffector/Source/ClonerEffector/Private/Utilities/CEClonerLayoutLatentAction.h
@@ -0,0 +1,41 @@
+// Copyright Epic Games, Inc. All Rights Reserved.
+
+#pragma once
+
+#include "LatentActions.h"
+#include "Templates/SubclassOf.h"
+#include "UObject/WeakObjectPtrFwd.h"
+
+class UCEClonerLayoutBase;
+class UCEClonerComponent;
+struct FLatentActionInfo;
+
+/** Latent action to set cloner layout and wait until layout is active and ready to use */
+class FCEClonerLayoutLatentAction : public FPendingLatentAction
+{
+public:
+	FCEClonerLayoutLatentAction(const FLatentActionInfo& InLatentInfo, UCEClonerComponent* InCloner, TSubclassOf<UCEClonerLayoutBase> InLayoutClass, bool& bOutSuccess, UCEClonerLayoutBase*& OutLayout);
+
+	virtual ~FCEClonerLayoutLatentAction();
+
+protected:
+	//~ Begin FPendingLatentAction
+	virtual void UpdateOperation(FLatentResponse& OutResponse) override;
+#if WITH_EDITOR
+	virtual FString GetDescription() const override;
+#endif
+	//~ End FPendingLatentAction
+
+	void OnClonerLayoutLoaded(UCEClonerComponent* InClonerComponent, UCEClonerLayoutBase* InClonerLayoutBase) const;
+
+	FName ExecutionFunction;
+	int32 OutputLink;
+	FWeakObjectPtr CallbackTarget;
+
+	bool bLayoutSet;
+	TWeakObjectPtr<UCEClonerComponent> ClonerWeak;
+	TSubclassOf<UCEClonerLayoutBase> LayoutClass;
+
+	bool& bSuccess;
+	UCEClonerLayoutBase*& Layout;
+};
diff --git a/Engine/Plugins/Experimental/ClonerEffector/Source/ClonerEffector/Private/Utilities/CEClonerLibrary.cpp b/Engine/Plugins/Experimental/ClonerEffector/Source/ClonerEffector/Private/Utilities/CEClonerLibrary.cpp
--- a/Engine/Plugins/Experimental/ClonerEffector/Source/ClonerEffector/Private/Utilities/CEClonerLibrary.cpp
+++ b/Engine/Plugins/Experimental/ClonerEffector/Source/ClonerEffector/Private/Utilities/CEClonerLibrary.cpp
@@ -2,7 +2,12 @@
 
 #include "Utilities/CEClonerLibrary.h"
 
+#include "CEClonerLayoutLatentAction.h"
+#include "Cloner/CEClonerComponent.h"
 #include "Containers/Set.h"
+#include "Engine/Engine.h"
+#include "Engine/LatentActionManager.h"
+#include "Engine/World.h"
 #include "Subsystems/CEClonerSubsystem.h"
 
 void UCEClonerLibrary::GetClonerLayoutClasses(TSet<TSubclassOf<UCEClonerLayoutBase>>& OutLayoutClasses)
@@ -24,3 +29,45 @@ void UCEClonerLibrary::GetClonerExtensionClasses(TSet<TSubclassOf<UCEClonerExten
 		OutExtensionClasses = ClonerSubsystem->GetExtensionClasses();
 	}
 }
+
+bool UCEClonerLibrary::GetClonerLayoutName(TSubclassOf<UCEClonerLayoutBase> InLayoutClass, FName& OutLayoutName)
+{
+	if (const UCEClonerSubsystem* ClonerSubsystem = UCEClonerSubsystem::Get())
+	{
+		OutLayoutName = ClonerSubsystem->FindLayoutName(InLayoutClass);
+	}
+
+	return !OutLayoutName.IsNone();
+}
+
+bool UCEClonerLibrary::GetClonerLayoutClass(FName InLayoutName, TSubclassOf<UCEClonerLayoutBase>& OutLayoutClass)
+{
+	if (const UCEClonerSubsystem* ClonerSubsystem = UCEClonerSubsystem::Get())
+	{
+		OutLayoutClass = ClonerSubsystem->FindLayoutClass(InLayoutName);
+	}
+
+	return !!OutLayoutClass.Get();
+}
+
+void UCEClonerLibrary::SetClonerLayoutByClass(const UObject* InWorldContext, FLatentActionInfo InLatentInfo, UCEClonerComponent* InCloner, TSubclassOf<UCEClonerLayoutBase> InLayoutClass, bool& bOutSuccess, UCEClonerLayoutBase*& OutLayout)
+{
+	if (UWorld* World = GEngine->GetWorldFromContextObject(InWorldContext, EGetWorldErrorMode::LogAndReturnNull))
+	{
+		FLatentActionManager& LatentActionManager = World->GetLatentActionManager();
+
+		if (!LatentActionManager.FindExistingAction<FCEClonerLayoutLatentAction>(InLatentInfo.CallbackTarget, InLatentInfo.UUID))
+		{
+			LatentActionManager.AddNewAction(InLatentInfo.CallbackTarget, InLatentInfo.UUID, new FCEClonerLayoutLatentAction(InLatentInfo, InCloner, InLayoutClass, bOutSuccess, OutLayout));
+		}
+	}
+}
+
+void UCEClonerLibrary::SetClonerLayoutByName(const UObject* InWorldContext, FLatentActionInfo InLatentInfo, UCEClonerComponent* InCloner, FName InLayoutName, bool& bOutSuccess, UCEClonerLayoutBase*& OutLayout)
+{
+	TSubclassOf<UCEClonerLayoutBase> LayoutClass;
+	if (UCEClonerLibrary::GetClonerLayoutClass(InLayoutName, LayoutClass))
+	{
+		UCEClonerLibrary::SetClonerLayoutByClass(InWorldContext, InLatentInfo, InCloner, LayoutClass, bOutSuccess, OutLayout);
+	}
+}
diff --git a/Engine/Plugins/Experimental/ClonerEffector/Source/ClonerEffector/Public/Subsystems/CEClonerSubsystem.h b/Engine/Plugins/Experimental/ClonerEffector/Source/ClonerEffector/Public/Subsystems/CEClonerSubsystem.h
--- a/Engine/Plugins/Experimental/ClonerEffector/Source/ClonerEffector/Public/Subsystems/CEClonerSubsystem.h
+++ b/Engine/Plugins/Experimental/ClonerEffector/Source/ClonerEffector/Public/Subsystems/CEClonerSubsystem.h
@@ -47,6 +47,9 @@ public:
 	/** Based on a layout class, find layout name */
 	FName FindLayoutName(TSubclassOf<UCEClonerLayoutBase> InLayoutClass) const;
 
+	/** Based on a layout name, find layout class */
+	TSubclassOf<UCEClonerLayoutBase> FindLayoutClass(FName InLayoutName) const;
+
 	/** Creates a new layout instance for a cloner */
 	UCEClonerLayoutBase* CreateNewLayout(FName InLayoutName, UCEClonerComponent* InCloner);
 
diff --git a/Engine/Plugins/Experimental/ClonerEffector/Source/ClonerEffector/Public/Utilities/CEClonerLibrary.h b/Engine/Plugins/Experimental/ClonerEffector/Source/ClonerEffector/Public/Utilities/CEClonerLibrary.h
--- a/Engine/Plugins/Experimental/ClonerEffector/Source/ClonerEffector/Public/Utilities/CEClonerLibrary.h
+++ b/Engine/Plugins/Experimental/ClonerEffector/Source/ClonerEffector/Public/Utilities/CEClonerLibrary.h
@@ -7,9 +7,10 @@
 #include "Templates/SubclassOf.h"
 #include "CEClonerLibrary.generated.h"
 
-template <typename T> class TSubclassOf;
+class UCEClonerComponent;
 class UCEClonerExtensionBase;
 class UCEClonerLayoutBase;
+struct FLatentActionInfo;
 
 /** Blueprint operations for cloner */
 UCLASS(MinimalAPI, DisplayName="Motion Design Cloner Library")
@@ -31,4 +32,47 @@ public:
 	 */
 	UFUNCTION(BlueprintCallable, Category="Motion Design|Cloner|Utility")
 	static CLONEREFFECTOR_API void GetClonerExtensionClasses(TSet<TSubclassOf<UCEClonerExtensionBase>>& OutExtensionClasses);
+
+	/**
+	 * Retrieves the layout name from a layout class
+	 * @param InLayoutClass Layout class to get the name from
+	 * @param OutLayoutName [Out] Layout name
+	 * @return true when the layout name was retrieved
+	 */
+	UFUNCTION(BlueprintCallable, Category="Motion Design|Cloner|Utility")
+	static CLONEREFFECTOR_API bool GetClonerLayoutName(TSubclassOf<UCEClonerLayoutBase> InLayoutClass, FName& OutLayoutName);
+
+	/**
+	 * Retrieves the layout class from a layout name
+	 * @param InLayoutName Layout name to get the class from
+	 * @param OutLayoutClass [Out] Layout class
+	 * @return true when the layout class was retrieved
+	 */
+	UFUNCTION(BlueprintCallable, Category="Motion Design|Cloner|Utility")
+	static CLONEREFFECTOR_API bool GetClonerLayoutClass(FName InLayoutName, TSubclassOf<UCEClonerLayoutBase>& OutLayoutClass);
+
+
+	/**
+	 * Sets the active layout of a cloner and wait until the layout is loaded and active
+	 * @param InWorldContext World context object
+	 * @param InLatentInfo Latent action info 
+	 * @param InCloner Target cloner component
+	 * @param InLayoutClass Cloner layout class
+	 * @param bOutSuccess [Out] True when the layout class is set
+	 * @param OutLayout [Out] Layout object corresponding to the layout class 
+	 */
+	UFUNCTION(BlueprintCallable, Category="Motion Design|Cloner|Utility", meta=(Latent, LatentInfo="InLatentInfo", WorldContext="InWorldContext"))
+	static CLONEREFFECTOR_API void SetClonerLayoutByClass(const UObject* InWorldContext, FLatentActionInfo InLatentInfo, UCEClonerComponent* InCloner, TSubclassOf<UCEClonerLayoutBase> InLayoutClass, bool& bOutSuccess, UCEClonerLayoutBase*& OutLayout);
+
+	/**
+	 * Sets the active layout of a cloner and wait until the layout is loaded and active
+	 * @param InWorldContext World context object
+	 * @param InLatentInfo Latent action info 
+	 * @param InCloner Target cloner component
+	 * @param InLayoutName Cloner layout name
+	 * @param bOutSuccess [Out] True when the layout class is set
+	 * @param OutLayout [Out] Layout object corresponding to the layout class 
+	 */
+	UFUNCTION(BlueprintCallable, Category="Motion Design|Cloner|Utility", meta=(Latent, LatentInfo="InLatentInfo", WorldContext="InWorldContext"))
+	static CLONEREFFECTOR_API void SetClonerLayoutByName(const UObject* InWorldContext, FLatentActionInfo InLatentInfo, UCEClonerComponent* InCloner, FName InLayoutName, bool& bOutSuccess, UCEClonerLayoutBase*& OutLayout);
 };
\ No newline at end of file
