MotionDesign : ClonerEffector
- Moving mesh builder into its own module within the plugin to be reused by other system
- Replacing dirty items by weak actors instead of raw items ptr

diff --git a/Engine/Plugins/Experimental/ClonerEffector/ClonerEffector.uplugin b/Engine/Plugins/Experimental/ClonerEffector/ClonerEffector.uplugin
--- a/Engine/Plugins/Experimental/ClonerEffector/ClonerEffector.uplugin
+++ b/Engine/Plugins/Experimental/ClonerEffector/ClonerEffector.uplugin
@@ -23,6 +23,11 @@
 			"Name": "ClonerEffectorEditor",
 			"Type": "Editor",
 			"LoadingPhase": "Default"
+		},
+		{
+			"Name": "ClonerEffectorMeshBuilder",
+			"Type": "Runtime",
+			"LoadingPhase": "Default"
 		}
 	],
 	"Plugins": [
diff --git a/Engine/Plugins/Experimental/ClonerEffector/Source/ClonerEffector/ClonerEffector.Build.cs b/Engine/Plugins/Experimental/ClonerEffector/Source/ClonerEffector/ClonerEffector.Build.cs
--- a/Engine/Plugins/Experimental/ClonerEffector/Source/ClonerEffector/ClonerEffector.Build.cs
+++ b/Engine/Plugins/Experimental/ClonerEffector/Source/ClonerEffector/ClonerEffector.Build.cs
@@ -11,11 +11,11 @@ public class ClonerEffector : ModuleRules
 		PublicDependencyModuleNames.AddRange(
 			new string[]
 			{
+				"ClonerEffectorMeshBuilder",
 				"Core",
 				"CoreUObject",
 				"Engine",
-				"Niagara",
-				"NiagaraCore",
+				"Niagara"
 			}
 		);
 
@@ -23,12 +23,10 @@ public class ClonerEffector : ModuleRules
 			new string[]
 			{
 				"DeveloperSettings",
-				"DynamicMesh",
 				"GeometryCore",
 				"GeometryFramework",
 				"GeometryScriptingCore",
-				"ProceduralMeshComponent",
-				"StaticMeshDescription"
+				"NiagaraCore"
 			}
 		);
 
diff --git a/Engine/Plugins/Experimental/ClonerEffector/Source/ClonerEffector/Private/CEClonerEffectorShared.cpp b/Engine/Plugins/Experimental/ClonerEffector/Source/ClonerEffector/Private/CEClonerEffectorShared.cpp
--- a/Engine/Plugins/Experimental/ClonerEffector/Source/ClonerEffector/Private/CEClonerEffectorShared.cpp
+++ b/Engine/Plugins/Experimental/ClonerEffector/Source/ClonerEffector/Private/CEClonerEffectorShared.cpp
@@ -151,7 +151,7 @@ AStaticMeshActor* UE::ClonerEffector::Conversion::ConvertClonerToStaticMesh(UCEC
 		return NewActor;
 	}
 
-	FCEClonerMeshBuilder ClonerMeshBuilder;
+	FCEMeshBuilder ClonerMeshBuilder;
 	const FTransform ClonerTransform = InCloner->GetComponentTransform();
 	if (!ClonerMeshBuilder.AppendComponent(InCloner, ClonerTransform))
 	{
@@ -212,7 +212,7 @@ AStaticMeshActor* UE::ClonerEffector::Conversion::ConvertClonerToStaticMesh(UCEC
 	}
 
 	TArray<TWeakObjectPtr<UMaterialInterface>> MeshMaterials;
-	FCEClonerMeshBuilder::FCEClonerMeshBuilderParams Params;
+	FCEMeshBuilder::FCEMeshBuilderParams Params;
 	Params.bMergeMaterials = true;
 	ClonerMeshBuilder.BuildStaticMesh(StaticMesh, MeshMaterials, Params);
 
@@ -267,7 +267,7 @@ ADynamicMeshActor* UE::ClonerEffector::Conversion::ConvertClonerToDynamicMesh(UC
 		return NewActor;
 	}
 
-	FCEClonerMeshBuilder ClonerMeshBuilder;
+	FCEMeshBuilder ClonerMeshBuilder;
 	const FTransform ClonerTransform = InCloner->GetComponentTransform();
 	if (!ClonerMeshBuilder.AppendComponent(InCloner, ClonerTransform))
 	{
@@ -297,7 +297,7 @@ ADynamicMeshActor* UE::ClonerEffector::Conversion::ConvertClonerToDynamicMesh(UC
 	UDynamicMeshComponent* DynamicMeshComponent = NewActor->GetDynamicMeshComponent();
 
 	TArray<TWeakObjectPtr<UMaterialInterface>> MeshMaterials;
-	FCEClonerMeshBuilder::FCEClonerMeshBuilderParams Params;
+	FCEMeshBuilder::FCEMeshBuilderParams Params;
 	Params.bMergeMaterials = true;
 	ClonerMeshBuilder.BuildDynamicMesh(DynamicMeshComponent->GetDynamicMesh(), MeshMaterials, Params);
 
@@ -335,7 +335,7 @@ TArray<AStaticMeshActor*> UE::ClonerEffector::Conversion::ConvertClonerToStaticM
 		return NewActors;
 	}
 
-	FCEClonerMeshBuilder ClonerMeshBuilder;
+	FCEMeshBuilder ClonerMeshBuilder;
 	const FTransform ClonerTransform = InCloner->GetComponentTransform();
 	if (!ClonerMeshBuilder.AppendComponent(InCloner, ClonerTransform))
 	{
@@ -399,7 +399,7 @@ TArray<AStaticMeshActor*> UE::ClonerEffector::Conversion::ConvertClonerToStaticM
 			continue;
 		}
 
-		TArray<FCEClonerMeshBuilder::FCEClonerMeshInstanceData> Instances;
+		TArray<FCEMeshBuilder::FCEMeshInstanceData> Instances;
 		ClonerMeshBuilder.BuildStaticMesh(MeshIndex, StaticMesh, Instances);
 
 		// Replace material references that are no assets to avoid save issue in new packages
@@ -415,7 +415,7 @@ TArray<AStaticMeshActor*> UE::ClonerEffector::Conversion::ConvertClonerToStaticM
 			StaticMesh->SetMaterial(Index, Material);
 		}
 
-		for (const FCEClonerMeshBuilder::FCEClonerMeshInstanceData& Instance : Instances)
+		for (const FCEMeshBuilder::FCEMeshInstanceData& Instance : Instances)
 		{
 			if (AStaticMeshActor* StaticMeshActor = World->SpawnActor<AStaticMeshActor>(AStaticMeshActor::StaticClass(), ClonerTransform, SpawnParameters))
 			{
@@ -452,7 +452,7 @@ TArray<AStaticMeshActor*> UE::ClonerEffector::Conversion::ConvertClonerToStaticM
 
 			UStaticMesh* StaticMesh = NewObject<UStaticMesh>(StaticMeshComponent);
 
-			FCEClonerMeshBuilder::FCEClonerMeshInstanceData MeshData;
+			FCEMeshBuilder::FCEMeshInstanceData MeshData;
 			ClonerMeshBuilder.BuildStaticMesh(Index, StaticMesh, MeshData);
 
 			StaticMeshComponent->SetStaticMesh(StaticMesh);
@@ -495,7 +495,7 @@ TArray<ADynamicMeshActor*> UE::ClonerEffector::Conversion::ConvertClonerToDynami
 		return NewActors;
 	}
 
-	FCEClonerMeshBuilder ClonerMeshBuilder;
+	FCEMeshBuilder ClonerMeshBuilder;
 	const FTransform ClonerTransform = InCloner->GetComponentTransform();
 	if (!ClonerMeshBuilder.AppendComponent(InCloner, ClonerTransform))
 	{
@@ -535,7 +535,7 @@ TArray<ADynamicMeshActor*> UE::ClonerEffector::Conversion::ConvertClonerToDynami
 		{
 			UDynamicMeshComponent* DynamicMeshComponent = DynamicMeshActor->GetDynamicMeshComponent();
 
-			FCEClonerMeshBuilder::FCEClonerMeshInstanceData MeshData;
+			FCEMeshBuilder::FCEMeshInstanceData MeshData;
 			ClonerMeshBuilder.BuildDynamicMesh(Index, DynamicMeshComponent->GetDynamicMesh(), MeshData);
 
 			DynamicMeshActor->SetActorTransform(MeshData.Transform);
@@ -580,7 +580,7 @@ TArray<AActor*> UE::ClonerEffector::Conversion::ConvertClonerToInstancedStaticMe
 		return NewActors;
 	}
 
-	FCEClonerMeshBuilder ClonerMeshBuilder;
+	FCEMeshBuilder ClonerMeshBuilder;
 	const FTransform ClonerTransform = InCloner->GetComponentTransform();
 	if (!ClonerMeshBuilder.AppendComponent(InCloner, ClonerTransform))
 	{
@@ -649,7 +649,7 @@ TArray<AActor*> UE::ClonerEffector::Conversion::ConvertClonerToInstancedStaticMe
 				continue;
 			}
 
-			TArray<FCEClonerMeshBuilder::FCEClonerMeshInstanceData> Instances;
+			TArray<FCEMeshBuilder::FCEMeshInstanceData> Instances;
 			ClonerMeshBuilder.BuildStaticMesh(MeshIndex, StaticMesh, Instances);
 
 #if WITH_EDITOR
@@ -669,7 +669,7 @@ TArray<AActor*> UE::ClonerEffector::Conversion::ConvertClonerToInstancedStaticMe
 
 			ISMComponent->SetStaticMesh(StaticMesh);
 
-			for (const FCEClonerMeshBuilder::FCEClonerMeshInstanceData& Instance : Instances)
+			for (const FCEMeshBuilder::FCEMeshInstanceData& Instance : Instances)
 			{
 				ISMComponent->AddInstance(Instance.Transform, /** WorldSpace */true);
 			}
diff --git a/Engine/Plugins/Experimental/ClonerEffector/Source/ClonerEffector/Private/Cloner/CEClonerComponent.cpp b/Engine/Plugins/Experimental/ClonerEffector/Source/ClonerEffector/Private/Cloner/CEClonerComponent.cpp
--- a/Engine/Plugins/Experimental/ClonerEffector/Source/ClonerEffector/Private/Cloner/CEClonerComponent.cpp
+++ b/Engine/Plugins/Experimental/ClonerEffector/Source/ClonerEffector/Private/Cloner/CEClonerComponent.cpp
@@ -270,7 +270,7 @@ void UCEClonerComponent::UpdateAttachmentTree()
 		{
 			if (ClonedItem->MeshStatus == ECEClonerAttachmentStatus::Outdated)
 			{
-				ClonerTree.DirtyItemAttachments.Add(ClonedItem);
+				ClonerTree.DirtyItemAttachments.Add(ClonedItem->ItemActor);
 				InvalidateBakedStaticMesh(ClonedActor);
 			}
 			bClonerMeshesDirty = true;
@@ -521,7 +521,7 @@ void UCEClonerComponent::OnMaterialChanged(UObject* InObject)
 
 	for (UPrimitiveComponent* PrimitiveComponent : PrimitiveComponents)
 	{
-		if (!PrimitiveComponent || !FCEClonerMeshBuilder::HasAnyGeometry(PrimitiveComponent))
+		if (!PrimitiveComponent || !FCEMeshBuilder::HasAnyGeometry(PrimitiveComponent))
 		{
 			continue;
 		}
@@ -577,7 +577,7 @@ void UCEClonerComponent::OnMeshChanged(UStaticMeshComponent*, AActor* InActor)
 
 		Item->MeshStatus = ECEClonerAttachmentStatus::Outdated;
         InvalidateBakedStaticMesh(InActor);
-		ClonerTree.DirtyItemAttachments.Add(Item);
+		ClonerTree.DirtyItemAttachments.Add(Item->ItemActor);
 	}
 }
 
@@ -658,12 +658,13 @@ void UCEClonerComponent::UpdateDirtyMeshesAsync()
 	}
 
 	bClonerMeshesUpdating = true;
-	TSet<FCEClonerAttachmentItem*> UpdateItems = ClonerTree.DirtyItemAttachments;
-	ClonerTree.DirtyItemAttachments.Empty();
+
+	TSet<TWeakObjectPtr<AActor>> DirtyAttachments = ClonerTree.DirtyItemAttachments;
+	ClonerTree.DirtyItemAttachments.Empty(DirtyAttachments.Num());
 
 	// Update baked dynamic meshes on other thread
 	TWeakObjectPtr<UCEClonerComponent> ThisWeak(this);
-	Async(EAsyncExecution::TaskGraph, [ThisWeak, UpdateItems]()
+	Async(EAsyncExecution::TaskGraph, [ThisWeak, DirtyAttachments]()
 	{
 		UCEClonerComponent* This = ThisWeak.Get();
 
@@ -674,9 +675,11 @@ void UCEClonerComponent::UpdateDirtyMeshesAsync()
 
 		// update actor baked dynamic meshes
 		bool bSuccess = true;
-		for (FCEClonerAttachmentItem* Item : UpdateItems)
+		for (const TWeakObjectPtr<AActor>& Attachment : DirtyAttachments)
 		{
-			if (!Item || !Item->ItemActor.IsValid())
+			AActor* DirtyActor = Attachment.Get();
+
+			if (!DirtyActor)
 			{
 				continue;
 			}
@@ -684,11 +687,11 @@ void UCEClonerComponent::UpdateDirtyMeshesAsync()
 			if (IsGarbageCollectingAndLockingUObjectHashTables())
 			{
 				bSuccess = false;
-				This->ClonerTree.DirtyItemAttachments.Add(Item);
+				This->ClonerTree.DirtyItemAttachments.Add(Attachment);
 				continue;
 			}
 
-			This->UpdateActorBakedDynamicMesh(Item->ItemActor.Get());
+			This->UpdateActorBakedDynamicMesh(DirtyActor);
 		}
 
 		// Create baked static mesh on main thread (required)
@@ -1496,7 +1499,7 @@ void UCEClonerComponent::OnRenderStateDirty(UActorComponent& InActorComponent)
 	}
 
 	// Does it contain geometry that we can convert
-	if (!FCEClonerMeshBuilder::IsComponentSupported(&InActorComponent))
+	if (!FCEMeshBuilder::IsComponentSupported(&InActorComponent))
 	{
 		return;
 	}
@@ -1515,7 +1518,7 @@ void UCEClonerComponent::OnRenderStateDirty(UActorComponent& InActorComponent)
 
 	Item->MeshStatus = ECEClonerAttachmentStatus::Outdated;
 	InvalidateBakedStaticMesh(Owner);
-	ClonerTree.DirtyItemAttachments.Add(Item);
+	ClonerTree.DirtyItemAttachments.Add(Item->ItemActor);
 }
 
 void UCEClonerComponent::OnComponentTransformed(USceneComponent* InComponent, EUpdateTransformFlags InFlags, ETeleportType InTeleport)
@@ -1534,13 +1537,13 @@ void UCEClonerComponent::OnComponentTransformed(USceneComponent* InComponent, EU
 		return;
 	}
 
-	bool bComponentSupported = FCEClonerMeshBuilder::IsComponentSupported(InComponent);
+	bool bComponentSupported = FCEMeshBuilder::IsComponentSupported(InComponent);
 
 	if (!bComponentSupported)
 	{
 		for (const TObjectPtr<USceneComponent>& ChildComponent : InComponent->GetAttachChildren())
 		{
-			if (FCEClonerMeshBuilder::IsComponentSupported(ChildComponent))
+			if (FCEMeshBuilder::IsComponentSupported(ChildComponent))
 			{
 				bComponentSupported = true;
 				break;
@@ -1564,7 +1567,7 @@ void UCEClonerComponent::OnComponentTransformed(USceneComponent* InComponent, EU
 
 	Item->MeshStatus = ECEClonerAttachmentStatus::Outdated;
 	InvalidateBakedStaticMesh(Owner);
-	ClonerTree.DirtyItemAttachments.Add(Item);
+	ClonerTree.DirtyItemAttachments.Add(Item->ItemActor);
 }
 
 UCEClonerLayoutBase* UCEClonerComponent::FindOrAddLayout(TSubclassOf<UCEClonerLayoutBase> InClass)
diff --git a/Engine/Plugins/Experimental/ClonerEffector/Source/ClonerEffector/Private/Cloner/CEClonerMeshBuilder.cpp b/Engine/Plugins/Experimental/ClonerEffector/Source/ClonerEffector/Private/Cloner/CEClonerMeshBuilder.cpp
deleted file mode 100644
diff --git a/Engine/Plugins/Experimental/ClonerEffector/Source/ClonerEffector/Public/CEClonerEffectorShared.h b/Engine/Plugins/Experimental/ClonerEffector/Source/ClonerEffector/Public/CEClonerEffectorShared.h
--- a/Engine/Plugins/Experimental/ClonerEffector/Source/ClonerEffector/Public/CEClonerEffectorShared.h
+++ b/Engine/Plugins/Experimental/ClonerEffector/Source/ClonerEffector/Public/CEClonerEffectorShared.h
@@ -351,7 +351,8 @@ struct FCEClonerAttachmentTree
 	ECEClonerAttachmentStatus Status = ECEClonerAttachmentStatus::Updated;
 
 	/** Attachment items that are dirty and need an update */
-	TSet<FCEClonerAttachmentItem*> DirtyItemAttachments;
+	UPROPERTY()
+	TSet<TWeakObjectPtr<AActor>> DirtyItemAttachments;
 
 	void Reset()
 	{
diff --git a/Engine/Plugins/Experimental/ClonerEffector/Source/ClonerEffector/Public/Cloner/CEClonerComponent.h b/Engine/Plugins/Experimental/ClonerEffector/Source/ClonerEffector/Public/Cloner/CEClonerComponent.h
--- a/Engine/Plugins/Experimental/ClonerEffector/Source/ClonerEffector/Public/Cloner/CEClonerComponent.h
+++ b/Engine/Plugins/Experimental/ClonerEffector/Source/ClonerEffector/Public/Cloner/CEClonerComponent.h
@@ -3,7 +3,7 @@
 #pragma once
 
 #include "CEClonerEffectorShared.h"
-#include "CEClonerMeshBuilder.h"
+#include "CEMeshBuilder.h"
 #include "CEPropertyChangeDispatcher.h"
 #include "Containers/Ticker.h"
 #include "Layouts/CEClonerLayoutBase.h"
@@ -422,7 +422,7 @@ private:
 	FCEClonerAttachmentTree ClonerTree;
 
 	UPROPERTY(Transient, DuplicateTransient, NonTransactional)
-	FCEClonerMeshBuilder MeshBuilder;
+	FCEMeshBuilder MeshBuilder;
 
 	/** Asset meshes needs update */
 	bool bClonerMeshesDirty = true;
diff --git a/Engine/Plugins/Experimental/ClonerEffector/Source/ClonerEffector/Public/Cloner/CEClonerMeshBuilder.h b/Engine/Plugins/Experimental/ClonerEffector/Source/ClonerEffector/Public/Cloner/CEClonerMeshBuilder.h
deleted file mode 100644
diff --git a/Engine/Plugins/Experimental/ClonerEffector/Source/ClonerEffectorMeshBuilder/ClonerEffectorMeshBuilder.Build.cs b/Engine/Plugins/Experimental/ClonerEffector/Source/ClonerEffectorMeshBuilder/ClonerEffectorMeshBuilder.Build.cs
new file mode 100644
--- /dev/null
+++ b/Engine/Plugins/Experimental/ClonerEffector/Source/ClonerEffectorMeshBuilder/ClonerEffectorMeshBuilder.Build.cs
@@ -0,0 +1,36 @@
+﻿// Copyright Epic Games, Inc. All Rights Reserved.
+
+using UnrealBuildTool;
+
+public class ClonerEffectorMeshBuilder : ModuleRules
+{
+    public ClonerEffectorMeshBuilder(ReadOnlyTargetRules Target) : base(Target)
+    {
+        PCHUsage = ModuleRules.PCHUsageMode.UseExplicitOrSharedPCHs;
+
+        PublicDependencyModuleNames.AddRange(
+            new string[]
+            {
+                "Core"
+            }
+        );
+
+        PrivateDependencyModuleNames.AddRange(
+            new string[]
+            {
+                "CoreUObject",
+                "DynamicMesh",
+                "Engine",
+                "GeometryCore",
+                "GeometryFramework",
+                "GeometryScriptingCore",
+                "Niagara",
+                "NiagaraCore",
+                "ProceduralMeshComponent",
+                "StaticMeshDescription"
+            }
+        );
+        
+        ShortName = "CEMB";
+    }
+}
\ No newline at end of file
diff --git a/Engine/Plugins/Experimental/ClonerEffector/Source/ClonerEffectorMeshBuilder/Private/CEMeshBuilder.cpp b/Engine/Plugins/Experimental/ClonerEffector/Source/ClonerEffectorMeshBuilder/Private/CEMeshBuilder.cpp
new file mode 100644
--- /dev/null
+++ b/Engine/Plugins/Experimental/ClonerEffector/Source/ClonerEffectorMeshBuilder/Private/CEMeshBuilder.cpp
@@ -0,0 +1,1039 @@
+// Copyright Epic Games, Inc. All Rights Reserved.
+
+#include "CEMeshBuilder.h"
+
+#include "Components/BrushComponent.h"
+#include "Components/DynamicMeshComponent.h"
+#include "Components/InstancedStaticMeshComponent.h"
+#include "Components/SkeletalMeshComponent.h"
+#include "Components/SplineMeshComponent.h"
+#include "DynamicMesh/MeshNormals.h"
+#include "DynamicMesh/MeshTransforms.h"
+#include "DynamicMesh/Operations/MergeCoincidentMeshEdges.h"
+#include "DynamicMeshEditor.h"
+#include "Engine/SkeletalMesh.h"
+#include "Engine/StaticMesh.h"
+#include "GeometryScript/MeshAssetFunctions.h"
+#include "GeometryScript/MeshBasicEditFunctions.h"
+#include "GeometryScript/SceneUtilityFunctions.h"
+#include "Materials/MaterialInterface.h"
+#include "Model.h"
+#include "NiagaraComponent.h"
+#include "NiagaraEmitter.h"
+#include "NiagaraMeshRendererProperties.h"
+#include "NiagaraSimCacheFunctionLibrary.h"
+#include "NiagaraSystem.h"
+#include "ProceduralMeshComponent.h"
+#include "Rendering/SkeletalMeshRenderData.h"
+#include "StaticMeshOperations.h"
+#include "UDynamicMesh.h"
+
+const FCEMeshBuilder::FCEMeshBuilderParams FCEMeshBuilder::DefaultParams;
+
+bool FCEMeshBuilder::HasAnyGeometry(UActorComponent* InComponent)
+{
+	if (!IsComponentSupported(InComponent))
+	{
+		return false;
+	}
+
+	if (UDynamicMeshComponent* DynamicMeshComponent = Cast<UDynamicMeshComponent>(InComponent))
+	{
+		return DynamicMeshComponent->GetDynamicMesh()
+			&& DynamicMeshComponent->GetDynamicMesh()->GetTriangleCount() > 0;
+	}
+
+	if (const USkeletalMeshComponent* SkeletalMeshComponent = Cast<USkeletalMeshComponent>(InComponent))
+	{
+		if (SkeletalMeshComponent->GetSkeletalMeshAsset())
+		{
+			if (FSkeletalMeshRenderData* RenderData = SkeletalMeshComponent->GetSkeletalMeshRenderData())
+			{
+				return !RenderData->LODRenderData.IsEmpty() && RenderData->LODRenderData[0].GetNumVertices() > 0;
+			}
+		}
+
+		return false;
+	}
+
+	if (const UBrushComponent* BrushComponent = Cast<UBrushComponent>(InComponent))
+	{
+		return BrushComponent->Brush
+			&& !BrushComponent->Brush->Verts.IsEmpty();
+	}
+
+	if (UProceduralMeshComponent* ProceduralMeshComponent = Cast<UProceduralMeshComponent>(InComponent))
+	{
+		for (int32 SectionIndex = 0; SectionIndex < ProceduralMeshComponent->GetNumSections(); SectionIndex++)
+		{
+			if (const FProcMeshSection* Section = ProceduralMeshComponent->GetProcMeshSection(SectionIndex))
+			{
+				if (Section->bSectionVisible && !Section->ProcVertexBuffer.IsEmpty() && !Section->ProcIndexBuffer.IsEmpty())
+				{
+					return true;
+				}
+			}
+		}
+
+		return false;
+	}
+
+	if (const UInstancedStaticMeshComponent* ISMComponent = Cast<UInstancedStaticMeshComponent>(InComponent))
+	{
+		if (ISMComponent->GetStaticMesh() && ISMComponent->GetStaticMesh()->GetNumTriangles(/**LOD*/0) > 0)
+		{
+			return ISMComponent->GetNumInstances() > 0;
+		}
+
+		return false;
+	}
+
+	if (const USplineMeshComponent* SplineMeshComponent = Cast<USplineMeshComponent>(InComponent))
+	{
+		return SplineMeshComponent->GetStaticMesh() && SplineMeshComponent->GetStaticMesh()->GetNumTriangles(/**LOD*/0) > 0;
+	}
+
+	if (const UStaticMeshComponent* StaticMeshComponent = Cast<UStaticMeshComponent>(InComponent))
+	{
+		return StaticMeshComponent->GetStaticMesh() && StaticMeshComponent->GetStaticMesh()->GetNumTriangles(/**LOD*/0) > 0;
+	}
+
+	if (const UNiagaraComponent* NiagaraComponent = Cast<UNiagaraComponent>(InComponent))
+	{
+		if (UNiagaraSystem* System = NiagaraComponent->GetAsset())
+		{
+			if (System->GetActiveInstancesCount() > 0)
+			{
+				for (const FNiagaraEmitterHandle& EmitterHandle : System->GetEmitterHandles())
+				{
+					if (const FVersionedNiagaraEmitterData* EmitterData = EmitterHandle.GetEmitterData())
+					{
+						for (UNiagaraRendererProperties* EmitterRenderer : EmitterData->GetRenderers())
+						{
+							if (UNiagaraMeshRendererProperties* MeshRenderer = Cast<UNiagaraMeshRendererProperties>(EmitterRenderer))
+							{
+								for (const FNiagaraMeshRendererMeshProperties& MeshProperty : MeshRenderer->Meshes)
+								{
+									if (MeshProperty.Mesh && MeshProperty.Mesh->GetNumTriangles(/**LOD*/0) > 0)
+									{
+										return true;
+									}
+								}
+							}
+						}
+					}
+				}
+			}
+		}
+
+		return false;
+	}
+
+	return false;
+}
+
+FCEMeshBuilder::FCEMeshBuilder()
+{
+	OutputDynamicMesh = NewObject<UDynamicMesh>();
+}
+
+TArray<uint32> FCEMeshBuilder::GetMeshIndexes() const
+{
+	TArray<uint32> MeshIndexes;
+	Meshes.GenerateKeyArray(MeshIndexes);
+	return MeshIndexes;
+}
+
+bool FCEMeshBuilder::IsActorSupported(const AActor* InActor)
+{
+	return InActor && InActor->FindComponentByClass<UPrimitiveComponent>();
+}
+
+bool FCEMeshBuilder::IsComponentSupported(const UActorComponent* InComponent)
+{
+	return InComponent && InComponent->IsA<UPrimitiveComponent>();
+}
+
+void FCEMeshBuilder::Reset()
+{
+	ClearOutputMesh();
+	Meshes.Empty();
+	MeshInstances.Empty();
+}
+
+int32 FCEMeshBuilder::AppendActor(const AActor* InActor)
+{
+	int32 ComponentConverted = 0;
+
+	if (!IsValid(InActor))
+	{
+		return ComponentConverted;
+	}
+
+	const FTransform& SourceTransform = InActor->GetActorTransform();
+
+	TArray<UPrimitiveComponent*> PrimitiveComponents;
+	InActor->GetComponents(PrimitiveComponents, /** IncludeChildrenActors */ false);
+
+	for (UPrimitiveComponent* PrimitiveComponent : PrimitiveComponents)
+	{
+		if (!IsValid(PrimitiveComponent))
+		{
+			continue;
+		}
+
+		if (UDynamicMeshComponent* DynamicMeshComponent = Cast<UDynamicMeshComponent>(PrimitiveComponent))
+		{
+			if (AppendComponent(DynamicMeshComponent, SourceTransform))
+			{
+				ComponentConverted++;
+			}
+		}
+		else if (USkeletalMeshComponent* SkeletalMeshComponent = Cast<USkeletalMeshComponent>(PrimitiveComponent))
+		{
+			if (AppendComponent(SkeletalMeshComponent, SourceTransform))
+			{
+				ComponentConverted++;
+			}
+		}
+		else if (UBrushComponent* BrushComponent = Cast<UBrushComponent>(PrimitiveComponent))
+		{
+			if (AppendComponent(BrushComponent, SourceTransform))
+			{
+				ComponentConverted++;
+			}
+		}
+		else if (UProceduralMeshComponent* ProceduralMeshComponent = Cast<UProceduralMeshComponent>(PrimitiveComponent))
+		{
+			if (AppendComponent(ProceduralMeshComponent, SourceTransform))
+			{
+				ComponentConverted++;
+			}
+		}
+		else if (UInstancedStaticMeshComponent* ISMComponent = Cast<UInstancedStaticMeshComponent>(PrimitiveComponent))
+		{
+			if (AppendComponent(ISMComponent, SourceTransform))
+			{
+				ComponentConverted++;
+			}
+		}
+		else if (USplineMeshComponent* SplineMeshComponent = Cast<USplineMeshComponent>(PrimitiveComponent))
+		{
+			if (AppendComponent(SplineMeshComponent, SourceTransform))
+			{
+				ComponentConverted++;
+			}
+		}
+		else if (UStaticMeshComponent* StaticMeshComponent = Cast<UStaticMeshComponent>(PrimitiveComponent))
+		{
+			if (AppendComponent(StaticMeshComponent, SourceTransform))
+			{
+				ComponentConverted++;
+			}
+		}
+		else if (UNiagaraComponent* NiagaraComponent = Cast<UNiagaraComponent>(PrimitiveComponent))
+		{
+			if (AppendComponent(NiagaraComponent, SourceTransform))
+			{
+				ComponentConverted++;
+			}
+		}
+	}
+
+	return ComponentConverted;
+}
+
+bool FCEMeshBuilder::AppendMesh(const UDynamicMesh* InMesh, const TArray<TWeakObjectPtr<UMaterialInterface>>& InMaterials, const FTransform& InTransform)
+{
+	if (!IsValid(InMesh) || InMesh->GetTriangleCount() == 0)
+	{
+		return false;
+	}
+
+	return !!AddMeshInstance(InMesh->GetUniqueID(), InTransform, InMaterials, [&InMesh](FDynamicMesh3& InCreateMesh)->bool
+	{
+		// Create a copy of the mesh
+		InMesh->ProcessMesh([&InCreateMesh](const FDynamicMesh3& EditMesh)
+		{
+			InCreateMesh = EditMesh;
+		});
+
+		return true;
+	});
+}
+
+bool FCEMeshBuilder::AppendMesh(UStaticMesh* InMesh, const TArray<TWeakObjectPtr<UMaterialInterface>>& InMaterials, const FTransform& InSourceTransform)
+{
+	if (!IsValid(InMesh) || InMesh->GetNumTriangles(/** LOD*/ 0) == 0)
+	{
+		return false;
+	}
+
+	return !!AddMeshInstance(InMesh->GetUniqueID(), InSourceTransform, InMaterials, [this, &InMesh](FDynamicMesh3& InCreateMesh)->bool
+	{
+		// convert to dynamic mesh
+		FGeometryScriptMeshReadLOD StaticMeshLOD;
+		StaticMeshLOD.LODType = EGeometryScriptLODType::RenderData;
+
+		FGeometryScriptCopyMeshFromAssetOptions OutputMeshOptions;
+		OutputMeshOptions.bIgnoreRemoveDegenerates = false;
+		OutputMeshOptions.bRequestTangents = false;
+		OutputMeshOptions.bApplyBuildSettings = false;
+
+		EGeometryScriptOutcomePins OutResult;
+		UGeometryScriptLibrary_StaticMeshFunctions::CopyMeshFromStaticMesh(InMesh, OutputDynamicMesh, OutputMeshOptions, StaticMeshLOD, OutResult);
+
+		if (OutResult != EGeometryScriptOutcomePins::Success)
+		{
+			return false;
+		}
+
+		OutputDynamicMesh->EditMesh([this, &InCreateMesh](FDynamicMesh3& EditMesh)
+		{
+			InCreateMesh = MoveTemp(EditMesh);
+
+			// replace by empty mesh
+			FDynamicMesh3 EmptyMesh;
+			EditMesh = MoveTemp(EmptyMesh);
+		}, EDynamicMeshChangeType::GeneralEdit, EDynamicMeshAttributeChangeFlags::Unknown, true);
+
+		return true;
+	});
+}
+
+bool FCEMeshBuilder::AppendComponent(const UStaticMeshComponent* InComponent, const FTransform& InSourceTransform)
+{
+	if (!IsValid(InComponent))
+	{
+		return false;
+	}
+
+	UStaticMesh* StaticMesh = InComponent->GetStaticMesh();
+
+	if (!IsValid(StaticMesh))
+	{
+		return false;
+	}
+
+	const FTransform RelativeTransform = InComponent->GetComponentTransform().GetRelativeTransform(InSourceTransform);
+
+	TArray<TWeakObjectPtr<UMaterialInterface>> Materials;
+	Algo::Transform(InComponent->GetMaterials(), Materials, [](UMaterialInterface* InMaterial)
+	{
+		return InMaterial;
+	});
+
+	return AppendMesh(StaticMesh, Materials, RelativeTransform);
+}
+
+bool FCEMeshBuilder::AppendComponent(UProceduralMeshComponent* InComponent, const FTransform& InSourceTransform)
+{
+	if (!IsValid(InComponent))
+	{
+		return false;
+	}
+
+	const int32 SectionCount = InComponent->GetNumSections();
+
+	if (SectionCount == 0)
+	{
+		return false;
+	}
+
+	// Transform the new mesh relative to the component
+	const FTransform RelativeTransform = InComponent->GetComponentTransform().GetRelativeTransform(InSourceTransform);
+
+	TArray<TWeakObjectPtr<UMaterialInterface>> Materials;
+	Algo::Transform(InComponent->GetMaterials(), Materials, [](UMaterialInterface* InMaterial)
+	{
+		return InMaterial;
+	});
+
+	const uint32 MeshIndex = InComponent->GetUniqueID();
+
+	return !!AddMeshInstance(MeshIndex, RelativeTransform, Materials, [this, &InComponent, &SectionCount](FDynamicMesh3& InCreateMesh)->bool
+	{
+		InCreateMesh.EnableAttributes();
+		InCreateMesh.Attributes()->EnablePrimaryColors();
+		InCreateMesh.Attributes()->EnableMaterialID();
+		InCreateMesh.Attributes()->SetNumNormalLayers(1);
+		InCreateMesh.Attributes()->SetNumUVLayers(1);
+		InCreateMesh.Attributes()->SetNumPolygroupLayers(1);
+		InCreateMesh.Attributes()->EnableTangents();
+
+		using namespace UE::Geometry;
+
+		FDynamicMeshColorOverlay* ColorOverlay = InCreateMesh.Attributes()->PrimaryColors();
+		FDynamicMeshNormalOverlay* NormalOverlay = InCreateMesh.Attributes()->PrimaryNormals();
+		FDynamicMeshUVOverlay* UVOverlay = InCreateMesh.Attributes()->PrimaryUV();
+		FDynamicMeshMaterialAttribute* MaterialAttr = InCreateMesh.Attributes()->GetMaterialID();
+		FDynamicMeshPolygroupAttribute* GroupAttr = InCreateMesh.Attributes()->GetPolygroupLayer(0);
+		FDynamicMeshNormalOverlay* TangentOverlay = InCreateMesh.Attributes()->PrimaryTangents();
+
+		for (int32 SectionIdx = 0; SectionIdx < SectionCount; SectionIdx++)
+		{
+			if (FProcMeshSection* Section = InComponent->GetProcMeshSection(SectionIdx))
+			{
+				if (Section->bSectionVisible)
+				{
+					TArray<int32> VtxIds;
+					TArray<int32> NormalIds;
+					TArray<int32> ColorIds;
+					TArray<int32> UVIds;
+					TArray<int32> TaIds;
+
+					// copy vertices data (position, normal, color, UV, tangent)
+					for (FProcMeshVertex& SectionVertex : Section->ProcVertexBuffer)
+					{
+						int32 VId = InCreateMesh.AppendVertex(SectionVertex.Position);
+						VtxIds.Add(VId);
+
+						int32 NId = NormalOverlay->AppendElement(static_cast<FVector3f>(SectionVertex.Normal));
+						NormalIds.Add(NId);
+
+						int32 CId = ColorOverlay->AppendElement(static_cast<FVector4f>(SectionVertex.Color));
+						ColorIds.Add(CId);
+
+						int32 UVId = UVOverlay->AppendElement(static_cast<FVector2f>(SectionVertex.UV0));
+						UVIds.Add(UVId);
+
+						int32 TaId = TangentOverlay->AppendElement(static_cast<FVector3f>(SectionVertex.Tangent.TangentX));
+						TaIds.Add(TaId);
+					}
+
+					// copy tris data
+					if (Section->ProcIndexBuffer.Num() % 3 != 0)
+					{
+						continue;
+					}
+
+					for (int32 Idx = 0; Idx < Section->ProcIndexBuffer.Num(); Idx+=3)
+					{
+						int32 VIdx1 = Section->ProcIndexBuffer[Idx];
+						int32 VIdx2 = Section->ProcIndexBuffer[Idx + 1];
+						int32 VIdx3 = Section->ProcIndexBuffer[Idx + 2];
+
+						int32 VId1 = VtxIds[VIdx1];
+						int32 VId2 = VtxIds[VIdx2];
+						int32 VId3 = VtxIds[VIdx3];
+
+						int32 TId = InCreateMesh.AppendTriangle(VId1, VId2, VId3, SectionIdx);
+
+						if (TId < 0)
+						{
+							continue;
+						}
+
+						NormalOverlay->SetTriangle(TId, FIndex3i(NormalIds[VIdx1], NormalIds[VIdx2], NormalIds[VIdx3]), true);
+						ColorOverlay->SetTriangle(TId, FIndex3i(ColorIds[VIdx1], ColorIds[VIdx2], ColorIds[VIdx3]), true);
+						UVOverlay->SetTriangle(TId, FIndex3i(UVIds[VIdx1], UVIds[VIdx2], UVIds[VIdx3]), true);
+						TangentOverlay->SetTriangle(TId, FIndex3i(TaIds[VIdx1], TaIds[VIdx2], TaIds[VIdx3]), true);
+
+						MaterialAttr->SetValue(TId, SectionIdx);
+						GroupAttr->SetValue(TId, SectionIdx);
+					}
+				}
+			}
+		}
+
+		return true;
+	});
+}
+
+bool FCEMeshBuilder::AppendComponent(UBrushComponent* InComponent, const FTransform& InSourceTransform)
+{
+	return AppendPrimitiveComponent(nullptr, InComponent, InSourceTransform);
+}
+
+bool FCEMeshBuilder::AppendComponent(const USkeletalMeshComponent* InComponent, const FTransform& InSourceTransform)
+{
+	if (!IsValid(InComponent))
+	{
+		return false;
+	}
+
+	USkeletalMesh* SkeletalMesh = InComponent->GetSkeletalMeshAsset();
+
+	if (!IsValid(SkeletalMesh))
+	{
+		return false;
+	}
+
+	// Transform the new mesh relative to the component
+	const FTransform RelativeTransform = InComponent->GetComponentTransform().GetRelativeTransform(InSourceTransform);
+
+	TArray<TWeakObjectPtr<UMaterialInterface>> Materials;
+	Algo::Transform(InComponent->GetMaterials(), Materials, [](UMaterialInterface* InMaterial)
+	{
+		return InMaterial;
+	});
+
+	return !!AddMeshInstance(SkeletalMesh->GetUniqueID(), RelativeTransform, Materials, [this, &SkeletalMesh](FDynamicMesh3& InCreateMesh)->bool
+	{
+		// convert to dynamic mesh
+		FGeometryScriptMeshReadLOD SkeletalMeshLOD;
+		SkeletalMeshLOD.LODType = EGeometryScriptLODType::SourceModel;
+
+		FGeometryScriptCopyMeshFromAssetOptions OutputMeshOptions;
+		OutputMeshOptions.bIgnoreRemoveDegenerates = false;
+		OutputMeshOptions.bRequestTangents = false;
+		OutputMeshOptions.bApplyBuildSettings = false;
+
+		EGeometryScriptOutcomePins OutResult;
+		UGeometryScriptLibrary_StaticMeshFunctions::CopyMeshFromSkeletalMesh(SkeletalMesh, OutputDynamicMesh, OutputMeshOptions, SkeletalMeshLOD, OutResult);
+
+		if (OutResult != EGeometryScriptOutcomePins::Success)
+		{
+			return false;
+		}
+
+		OutputDynamicMesh->EditMesh([&InCreateMesh](FDynamicMesh3& EditMesh)
+		{
+			InCreateMesh = MoveTemp(EditMesh);
+
+			// replace by empty mesh for next usage
+			FDynamicMesh3 EmptyMesh;
+			EditMesh = MoveTemp(EmptyMesh);
+		}, EDynamicMeshChangeType::GeneralEdit, EDynamicMeshAttributeChangeFlags::Unknown, /** bDeferChanges */true);
+
+		return true;
+	});
+}
+
+bool FCEMeshBuilder::AppendComponent(UDynamicMeshComponent* InComponent, const FTransform& InSourceTransform)
+{
+	if (!IsValid(InComponent))
+	{
+		return false;
+	}
+
+	const UDynamicMesh* DynamicMesh = InComponent->GetDynamicMesh();
+
+	if (!IsValid(DynamicMesh) || DynamicMesh->GetTriangleCount() == 0)
+	{
+		return false;
+	}
+
+	// Transform the new mesh relative to the component
+	const FTransform RelativeTransform = InComponent->GetComponentTransform().GetRelativeTransform(InSourceTransform);
+
+	// Copy all materials
+	TArray<TWeakObjectPtr<UMaterialInterface>> Materials;
+	Algo::Transform(InComponent->GetMaterials(), Materials, [](UMaterialInterface* InMaterial)
+	{
+		return InMaterial;
+	});
+
+	return AppendMesh(DynamicMesh, Materials, RelativeTransform);
+}
+
+bool FCEMeshBuilder::AppendComponent(UInstancedStaticMeshComponent* InComponent, const FTransform& InSourceTransform)
+{
+	if (!IsValid(InComponent) || !IsValid(InComponent->GetStaticMesh()))
+	{
+		return false;
+	}
+
+	return AppendPrimitiveComponent(nullptr, InComponent, InSourceTransform);
+}
+
+bool FCEMeshBuilder::AppendComponent(USplineMeshComponent* InComponent, const FTransform& InSourceTransform)
+{
+	if (!IsValid(InComponent) || !IsValid(InComponent->GetStaticMesh()))
+	{
+		return false;
+	}
+
+	return AppendPrimitiveComponent(nullptr, InComponent, InSourceTransform);
+}
+
+bool FCEMeshBuilder::AppendComponent(UNiagaraComponent* InComponent, const FTransform& InSourceTransform)
+{
+	if (!IsValid(InComponent))
+	{
+		return false;
+	}
+
+	UNiagaraSystem* System = InComponent->GetAsset();
+
+	if (!IsValid(System))
+	{
+		return false;
+	}
+
+	struct FNiagaraSimCacheEmitterData
+	{
+		TArray<UNiagaraMeshRendererProperties*> MeshRenderers;
+		TArray<FVector> ParticlePositions;
+		TArray<FQuat> ParticleRotations;
+		TArray<FVector> ParticleScales;
+		TArray<int32> ParticleMeshIndexes;
+	};
+
+	TMap<FName, FNiagaraSimCacheEmitterData> EmittersData;
+
+	// Set attributes to capture
+	FNiagaraSimCacheCreateParameters Params;
+	Params.AttributeCaptureMode = ENiagaraSimCacheAttributeCaptureMode::ExplicitAttributes;
+	Params.bAllowDataInterfaceCaching = false;
+
+	for (const FNiagaraEmitterHandle& EmitterHandle : System->GetEmitterHandles())
+	{
+		const FString EmitterName = EmitterHandle.GetUniqueInstanceName();
+
+		Params.ExplicitCaptureAttributes.Add(FName(EmitterName + TEXT(".Particles.Position")));
+		Params.ExplicitCaptureAttributes.Add(FName(EmitterName + TEXT(".Particles.MeshOrientation")));
+		Params.ExplicitCaptureAttributes.Add(FName(EmitterName + TEXT(".Particles.Scale")));
+		Params.ExplicitCaptureAttributes.Add(FName(EmitterName + TEXT(".Particles.MeshIndex")));
+
+		if (const FVersionedNiagaraEmitterData* EmitterData = EmitterHandle.GetEmitterData())
+		{
+			for (UNiagaraRendererProperties* EmitterRenderer : EmitterData->GetRenderers())
+			{
+				if (UNiagaraMeshRendererProperties* MeshRenderer = Cast<UNiagaraMeshRendererProperties>(EmitterRenderer))
+				{
+					FNiagaraSimCacheEmitterData& CacheEmitterData = EmittersData.FindOrAdd(FName(EmitterName));
+					CacheEmitterData.MeshRenderers.Add(MeshRenderer);
+				}
+			}
+		}
+	}
+
+	if (EmittersData.IsEmpty())
+	{
+		return false;
+	}
+
+	UNiagaraSimCache* SimCache = UNiagaraSimCacheFunctionLibrary::CreateNiagaraSimCache(InComponent);
+
+	if (!IsValid(SimCache))
+	{
+		return false;
+	}
+
+	const bool bSuccess = UNiagaraSimCacheFunctionLibrary::CaptureNiagaraSimCacheImmediate(SimCache, Params, InComponent, SimCache, /** AdvanceSim */ false);
+
+	if (!bSuccess)
+	{
+		return false;
+	}
+
+	for (const FNiagaraEmitterHandle& Handle : System->GetEmitterHandles())
+	{
+		constexpr int32 FrameIndex = 0;
+		constexpr bool bLocalToWorld = false;
+		const FName EmitterName(Handle.GetUniqueInstanceName());
+
+		FNiagaraSimCacheEmitterData& CacheEmitterData = EmittersData.FindChecked(EmitterName);
+
+		SimCache->ReadPositionAttribute(CacheEmitterData.ParticlePositions, TEXT("Position"), EmitterName, bLocalToWorld, FrameIndex);
+		SimCache->ReadQuatAttribute(CacheEmitterData.ParticleRotations, TEXT("MeshOrientation"), EmitterName, bLocalToWorld, FrameIndex);
+		SimCache->ReadVectorAttribute(CacheEmitterData.ParticleScales, TEXT("Scale"), EmitterName, FrameIndex);
+		SimCache->ReadIntAttribute(CacheEmitterData.ParticleMeshIndexes, TEXT("MeshIndex"), EmitterName, FrameIndex);
+
+		if (!(CacheEmitterData.ParticlePositions.Num() == CacheEmitterData.ParticleRotations.Num()
+			&& CacheEmitterData.ParticleRotations.Num() == CacheEmitterData.ParticleScales.Num()
+			&& CacheEmitterData.ParticleScales.Num() == CacheEmitterData.ParticleMeshIndexes.Num()))
+		{
+			return false;
+		}
+	}
+
+	const FTransform RelativeTransform = InComponent->GetComponentTransform().GetRelativeTransform(InSourceTransform);
+
+	bool bResult = false;
+
+	for (const TPair<FName, FNiagaraSimCacheEmitterData>& CacheEmitterDataPair : EmittersData)
+	{
+		const FNiagaraSimCacheEmitterData& CacheEmitterData = CacheEmitterDataPair.Value;
+
+		for (int32 Index = 0; Index < CacheEmitterData.ParticlePositions.Num(); Index++)
+		{
+			const FVector& ParticlePosition = CacheEmitterData.ParticlePositions[Index];
+			const FVector& ParticleScale = CacheEmitterData.ParticleScales[Index];
+			const FQuat& ParticleRotation = CacheEmitterData.ParticleRotations[Index];
+			const int32& ParticleMeshIndex = CacheEmitterData.ParticleMeshIndexes[Index];
+
+			FTransform ParticleTransform(ParticleRotation, ParticlePosition, ParticleScale);
+
+			if (ParticleMeshIndex >= 0)
+			{
+				for (UNiagaraMeshRendererProperties* MeshRenderer : CacheEmitterData.MeshRenderers)
+				{
+					if (!IsValid(MeshRenderer))
+					{
+						continue;
+					}
+
+					if (!MeshRenderer->Meshes.IsValidIndex(ParticleMeshIndex))
+					{
+						continue;
+					}
+
+					FNiagaraMeshRendererMeshProperties& Mesh = MeshRenderer->Meshes[ParticleMeshIndex];
+					if (!IsValid(Mesh.Mesh))
+					{
+						continue;
+					}
+
+					FTransform MeshTransform(Mesh.Rotation.Quaternion(), Mesh.PivotOffset, Mesh.Scale);
+					MeshTransform.Accumulate(ParticleTransform);
+					MeshTransform.Accumulate(RelativeTransform);
+
+					TArray<TWeakObjectPtr<UMaterialInterface>> Materials;
+					for (int32 SectionIndex = 0; SectionIndex < Mesh.Mesh->GetNumSections(0); SectionIndex++)
+					{
+						Materials.Add(Mesh.Mesh->GetMaterial(SectionIndex));
+					}
+
+					bResult |= AppendMesh(Mesh.Mesh, Materials, MeshTransform);
+				}
+			}
+		}
+	}
+
+	SimCache->MarkAsGarbage();
+
+	return bResult;
+}
+
+bool FCEMeshBuilder::BuildDynamicMesh(UDynamicMesh* OutMesh, TArray<TWeakObjectPtr<UMaterialInterface>>& OutMaterials, const FCEMeshBuilderParams& InParams)
+{
+	if (!IsValid(OutMesh))
+	{
+		return false;
+	}
+
+	OutMaterials.Empty();
+
+	// Lets combine all meshes components from this actor together
+	OutMesh->EditMesh([this, &OutMaterials, &InParams](FDynamicMesh3& InMergedMesh)
+	{
+		using namespace UE::Geometry;
+
+		InMergedMesh.Clear();
+		InMergedMesh.EnableAttributes();
+		InMergedMesh.Attributes()->SetNumNormalLayers(1);
+
+		FDynamicMeshEditor Editor(&InMergedMesh);
+		FGeometryScriptAppendMeshOptions AppendOptions;
+		AppendOptions.CombineMode = EGeometryScriptCombineAttributesMode::EnableAllMatching;
+		int32 MaterialCount = 0;
+
+		// Mesh index to material forward map to merge same meshes materials into same slot
+		TMap<uint32, TMap<int32, int32>> MeshToMaterialMap;
+
+		// Convert meshes
+		for (const FCEMeshInstance& MeshInstance : MeshInstances)
+		{
+			FDynamicMesh3& Mesh = Meshes[MeshInstance.MeshIndex];
+
+			if (Mesh.TriangleCount() == 0)
+			{
+				continue;
+			}
+
+			// Apply transform
+			FDynamicMesh3 ConvertedMesh = Mesh;
+			MeshTransforms::ApplyTransform(ConvertedMesh, MeshInstance.MeshData.Transform);
+
+			// Get materials
+			if (!InParams.bMergeMaterials || !MeshToMaterialMap.Contains(MeshInstance.MeshIndex))
+			{
+				OutMaterials.Append(MeshInstance.MeshData.MeshMaterials);
+			}
+
+			// Enable matching attributes & append mesh
+			FMeshIndexMappings TmpMappings;
+			AppendOptions.UpdateAttributesForCombineMode(InMergedMesh, ConvertedMesh);
+			Editor.AppendMesh(&ConvertedMesh, TmpMappings);
+
+			// Fix triangles materials linking
+			if (ConvertedMesh.HasAttributes() && ConvertedMesh.Attributes()->HasMaterialID())
+			{
+				const FDynamicMeshMaterialAttribute* FromMaterialIDAttrib = ConvertedMesh.Attributes()->GetMaterialID();
+				FDynamicMeshMaterialAttribute* ToMaterialIDAttrib = InMergedMesh.Attributes()->GetMaterialID();
+
+				TMap<int32, int32>& MaterialMap = MeshToMaterialMap.FindOrAdd(MeshInstance.MeshIndex);
+				for (const TPair<int32, int32>& FromToTId : TmpMappings.GetTriangleMap().GetForwardMap())
+				{
+					const int32 FromMatId = FromMaterialIDAttrib->GetValue(FromToTId.Key);
+					int32 ToMatId = FromMatId + MaterialCount;
+
+					// used to merge materials for same mesh index
+					if (const int32* MatId = MaterialMap.Find(FromMatId))
+					{
+						ToMatId = *MatId;
+					}
+
+					MaterialMap.Add(FromMatId, ToMatId);
+					ToMaterialIDAttrib->SetNewValue(FromToTId.Value, ToMatId);
+				}
+
+				MaterialCount += MaterialMap.Num();
+
+				if (!InParams.bMergeMaterials)
+				{
+					MaterialMap.Empty();
+				}
+			}
+		}
+
+		if (InMergedMesh.TriangleCount() > 0)
+		{
+			// Merge shared edges
+			FMergeCoincidentMeshEdges WeldOp(&InMergedMesh);
+			WeldOp.Apply();
+		}
+	}, EDynamicMeshChangeType::GeneralEdit, EDynamicMeshAttributeChangeFlags::Unknown, /** bDeferChange */true);
+
+	return true;
+}
+
+bool FCEMeshBuilder::BuildStaticMesh(UStaticMesh* OutMesh, TArray<TWeakObjectPtr<UMaterialInterface>>& OutMaterials, const FCEMeshBuilderParams& InParams)
+{
+	if (!IsValid(OutMesh))
+	{
+		return false;
+	}
+
+	ClearOutputMesh();
+
+	if (!BuildDynamicMesh(OutputDynamicMesh, OutMaterials, InParams))
+	{
+		return false;
+	}
+
+	return DynamicMeshToStaticMesh(OutputDynamicMesh, OutMesh, OutMaterials);
+}
+
+bool FCEMeshBuilder::BuildStaticMesh(int32 InInstanceIndex, UStaticMesh* OutMesh, FCEMeshInstanceData& OutMeshInstance)
+{
+	if (!IsValid(OutMesh) || !MeshInstances.IsValidIndex(InInstanceIndex))
+	{
+		return false;
+	}
+
+	FCEMeshInstance& MeshInstance = MeshInstances[InInstanceIndex];
+	FDynamicMesh3& Mesh = Meshes[MeshInstance.MeshIndex];
+
+	if (Mesh.TriangleCount() == 0)
+	{
+		return false;
+	}
+
+	OutMeshInstance.MeshMaterials.Empty(MeshInstance.MeshData.MeshMaterials.Num());
+	OutMeshInstance.MeshMaterials.Append(MeshInstance.MeshData.MeshMaterials);
+
+	OutMeshInstance.Transform = MeshInstance.MeshData.Transform;
+
+	OutputDynamicMesh->SetMesh(Mesh);
+
+	return DynamicMeshToStaticMesh(OutputDynamicMesh, OutMesh, OutMeshInstance.MeshMaterials);
+}
+
+bool FCEMeshBuilder::BuildDynamicMesh(int32 InInstanceIndex, UDynamicMesh* OutMesh, FCEMeshInstanceData& OutMeshInstance)
+{
+	if (!IsValid(OutMesh) || !MeshInstances.IsValidIndex(InInstanceIndex))
+	{
+		return false;
+	}
+
+	FCEMeshInstance& MeshInstance = MeshInstances[InInstanceIndex];
+	FDynamicMesh3& Mesh = Meshes[MeshInstance.MeshIndex];
+
+	if (Mesh.TriangleCount() == 0)
+	{
+		return false;
+	}
+
+	OutMeshInstance.MeshMaterials.Empty(MeshInstance.MeshData.MeshMaterials.Num());
+	OutMeshInstance.MeshMaterials.Append(MeshInstance.MeshData.MeshMaterials);
+
+	OutMeshInstance.Transform = MeshInstance.MeshData.Transform;
+
+	OutMesh->SetMesh(Mesh);
+
+	return true;
+}
+
+bool FCEMeshBuilder::BuildStaticMesh(uint32 InMeshIndex, UStaticMesh* OutMesh, TArray<FCEMeshInstanceData>& OutMeshInstances)
+{
+	const FDynamicMesh3* Mesh = Meshes.Find(InMeshIndex);
+
+	if (!IsValid(OutMesh) || !Mesh)
+	{
+		return false;
+	}
+
+	OutMeshInstances.Empty(MeshInstances.Num());
+	for (const FCEMeshInstance& MeshInstance : MeshInstances)
+	{
+		if (MeshInstance.MeshIndex == InMeshIndex)
+		{
+			OutMeshInstances.Add(MeshInstance.MeshData);
+		}
+	}
+
+	if (OutMeshInstances.IsEmpty())
+	{
+		return false;
+	}
+
+	OutputDynamicMesh->SetMesh(*Mesh);
+
+	return DynamicMeshToStaticMesh(OutputDynamicMesh, OutMesh, OutMeshInstances[0].MeshMaterials);
+}
+
+bool FCEMeshBuilder::BuildDynamicMesh(uint32 InMeshIndex, UDynamicMesh* OutMesh, TArray<FCEMeshInstanceData>& OutMeshInstances)
+{
+	const FDynamicMesh3* Mesh = Meshes.Find(InMeshIndex);
+
+	if (!IsValid(OutMesh) || !Mesh)
+	{
+		return false;
+	}
+
+	OutMeshInstances.Empty(MeshInstances.Num());
+	for (const FCEMeshInstance& MeshInstance : MeshInstances)
+	{
+		if (MeshInstance.MeshIndex == InMeshIndex)
+		{
+			OutMeshInstances.Add(MeshInstance.MeshData);
+		}
+	}
+
+	if (OutMeshInstances.IsEmpty())
+	{
+		return false;
+	}
+
+	OutMesh->SetMesh(*Mesh);
+
+	return true;
+}
+
+bool FCEMeshBuilder::AppendPrimitiveComponent(const UObject* InMeshObject, UPrimitiveComponent* InComponent, const FTransform& InSourceTransform)
+{
+	if (!IsValid(InComponent))
+	{
+		return false;
+	}
+
+	// Transform the new mesh relative to the component
+	const FTransform RelativeTransform = InComponent->GetComponentTransform().GetRelativeTransform(InSourceTransform);
+
+	// Copy all materials
+	TArray<TWeakObjectPtr<UMaterialInterface>> Materials;
+	Materials.Reserve(InComponent->GetNumMaterials());
+	for (int32 Index = 0; Index < InComponent->GetNumMaterials(); Index++)
+	{
+		Materials.Add(InComponent->GetMaterial(Index));
+	}
+
+	// Take mesh id or component id to find already converted mesh
+	const uint32 MeshIndex = InMeshObject ? InMeshObject->GetUniqueID() : InComponent->GetUniqueID();
+
+	return !!AddMeshInstance(MeshIndex, RelativeTransform, Materials, [this, &InComponent](FDynamicMesh3& InCreateMesh)->bool
+	{
+		constexpr FGeometryScriptCopyMeshFromComponentOptions Options;
+		FTransform LocalToWorld = FTransform::Identity;
+
+		// convert to dynamic mesh
+		EGeometryScriptOutcomePins OutResult;
+		UGeometryScriptLibrary_SceneUtilityFunctions::CopyMeshFromComponent(InComponent, OutputDynamicMesh, Options, false, LocalToWorld, OutResult);
+
+		if (OutResult != EGeometryScriptOutcomePins::Success)
+		{
+			return false;
+		}
+
+		OutputDynamicMesh->EditMesh([&InCreateMesh](FDynamicMesh3& EditMesh)
+		{
+			InCreateMesh = MoveTemp(EditMesh);
+
+			// replace by empty mesh
+			FDynamicMesh3 EmptyMesh;
+			EditMesh = MoveTemp(EmptyMesh);
+		}, EDynamicMeshChangeType::GeneralEdit, EDynamicMeshAttributeChangeFlags::Unknown, /** bDeferChange */true);
+
+		return true;
+	});
+}
+
+bool FCEMeshBuilder::DynamicMeshToStaticMesh(UDynamicMesh* InMesh, UStaticMesh* OutMesh, const TArray<TWeakObjectPtr<UMaterialInterface>>& InMaterials)
+{
+	TArray<TObjectPtr<UMaterialInterface>> NewMaterials;
+	Algo::Transform(InMaterials, NewMaterials, [](const TWeakObjectPtr<UMaterialInterface>& InMaterialWeak)
+	{
+		return InMaterialWeak.Get();
+	});
+
+	// Export options
+	FGeometryScriptCopyMeshToAssetOptions AssetOptions;
+	AssetOptions.bReplaceMaterials = true;
+	AssetOptions.bEnableRecomputeNormals = false;
+	AssetOptions.bEnableRecomputeTangents = false;
+	AssetOptions.bEnableRemoveDegenerates = true;
+	AssetOptions.bEmitTransaction = false;
+	AssetOptions.bApplyNaniteSettings = false;
+	AssetOptions.bDeferMeshPostEditChange = false;
+	AssetOptions.NewMaterials = NewMaterials;
+
+	// LOD options
+	FGeometryScriptMeshWriteLOD TargetLOD;
+	TargetLOD.LODIndex = 0;
+	TargetLOD.bWriteHiResSource = false;
+
+	// Convert merged mesh to static mesh
+	EGeometryScriptOutcomePins OutResult;
+	UGeometryScriptLibrary_StaticMeshFunctions::CopyMeshToStaticMesh(InMesh, OutMesh, AssetOptions, TargetLOD, OutResult);
+
+#if WITH_EDITORONLY_DATA
+	// Compute normals and tangents
+	FMeshDescription* MeshDescription = OutMesh->GetMeshDescription(TargetLOD.LODIndex);
+	FStaticMeshOperations::ComputeTriangleTangentsAndNormals(*MeshDescription);
+	FStaticMeshOperations::ComputeTangentsAndNormals(*MeshDescription, EComputeNTBsFlags::Normals);
+#endif
+
+	return OutResult == EGeometryScriptOutcomePins::Success;
+}
+
+void FCEMeshBuilder::ClearOutputMesh() const
+{
+	if (OutputDynamicMesh)
+	{
+		OutputDynamicMesh->EditMesh([](FDynamicMesh3& EditMesh){ EditMesh.Clear(); });
+	}
+}
+
+FCEMeshBuilder::FCEMeshInstance* FCEMeshBuilder::AddMeshInstance(uint32 InMeshIndex, const FTransform& InTransform, const TArray<TWeakObjectPtr<UMaterialInterface>>& InMaterials, TFunctionRef<bool(UE::Geometry::FDynamicMesh3&)> InCreateMeshFunction)
+{
+	FCEMeshInstance MeshInstance;
+	MeshInstance.MeshIndex = InMeshIndex;
+	MeshInstance.MeshData.Transform = InTransform;
+	MeshInstance.MeshData.MeshMaterials = InMaterials;
+
+	FDynamicMesh3* CachedMesh = Meshes.Find(InMeshIndex);
+
+	if (!CachedMesh)
+	{
+		FDynamicMesh3 Mesh;
+
+		ClearOutputMesh();
+
+		if (InCreateMeshFunction(Mesh) && Mesh.TriangleCount() > 0)
+		{
+			CachedMesh = &Meshes.Add(InMeshIndex, MoveTemp(Mesh));
+		}
+
+		ClearOutputMesh();
+	}
+
+	if (CachedMesh)
+	{
+		return &MeshInstances.Add_GetRef(MeshInstance);
+	}
+
+	return nullptr;
+}
diff --git a/Engine/Plugins/Experimental/ClonerEffector/Source/ClonerEffectorMeshBuilder/Private/CEMeshBuilderModule.cpp b/Engine/Plugins/Experimental/ClonerEffector/Source/ClonerEffectorMeshBuilder/Private/CEMeshBuilderModule.cpp
new file mode 100644
--- /dev/null
+++ b/Engine/Plugins/Experimental/ClonerEffector/Source/ClonerEffectorMeshBuilder/Private/CEMeshBuilderModule.cpp
@@ -0,0 +1,5 @@
+﻿// Copyright Epic Games, Inc. All Rights Reserved.
+
+#include "Modules/ModuleManager.h"
+
+IMPLEMENT_MODULE(FDefaultModuleImpl, ClonerEffectorMeshBuilder)
\ No newline at end of file
diff --git a/Engine/Plugins/Experimental/ClonerEffector/Source/ClonerEffectorMeshBuilder/Public/CEMeshBuilder.h b/Engine/Plugins/Experimental/ClonerEffector/Source/ClonerEffectorMeshBuilder/Public/CEMeshBuilder.h
new file mode 100644
--- /dev/null
+++ b/Engine/Plugins/Experimental/ClonerEffector/Source/ClonerEffectorMeshBuilder/Public/CEMeshBuilder.h
@@ -0,0 +1,147 @@
+// Copyright Epic Games, Inc. All Rights Reserved.
+
+#pragma once
+
+#include "DynamicMesh/DynamicMesh3.h"
+#include "UObject/WeakObjectPtrTemplates.h"
+#include "CEMeshBuilder.generated.h"
+
+class AActor;
+class UActorComponent;
+class UBrushComponent;
+class UDynamicMesh;
+class UDynamicMeshComponent;
+class UInstancedStaticMeshComponent;
+class UMaterialInterface;
+class UNiagaraComponent;
+class UProceduralMeshComponent;
+class USkeletalMeshComponent;
+class USplineMeshComponent;
+class UStaticMesh;
+class UStaticMeshComponent;
+
+/** Struct used to build a mesh based out of other meshes class */
+USTRUCT()
+struct CLONEREFFECTORMESHBUILDER_API FCEMeshBuilder
+{
+	GENERATED_BODY()
+
+	struct FCEMeshInstanceData
+	{
+		/** Transform to apply on the mesh instance */
+		FTransform Transform;
+
+		/** Materials applied on this instance */
+		TArray<TWeakObjectPtr<UMaterialInterface>> MeshMaterials;
+	};
+
+	struct FCEMeshBuilderParams
+	{
+		/** Merge same meshes material slot in the final result */
+		bool bMergeMaterials = false;
+	};
+
+	static const FCEMeshBuilderParams DefaultParams;
+
+	/** Checks if the component contains any geometry data */
+	static bool HasAnyGeometry(UActorComponent* InComponent);
+
+	/** Does the mesh builder supports this actor */
+	static bool IsActorSupported(const AActor* InActor);
+
+	/** Does the mesh builder supports this component */
+	static bool IsComponentSupported(const UActorComponent* InComponent);
+
+	FCEMeshBuilder();
+
+	int32 GetMeshInstanceCount() const
+	{
+		return MeshInstances.Num();
+	}
+
+	int32 GetMeshCount() const
+	{
+		return Meshes.Num();
+	}
+
+	TArray<uint32> GetMeshIndexes() const;
+
+	/** Resets the builder and clears data */
+	void Reset();
+
+	/** Appends supported components within actor */
+	int32 AppendActor(const AActor* InActor);
+
+	/** Appends Static Mesh component */
+	bool AppendComponent(const UStaticMeshComponent* InComponent, const FTransform& InTransform = FTransform::Identity);
+
+	/** Appends Procedural Mesh component */
+	bool AppendComponent(UProceduralMeshComponent* InComponent, const FTransform& InTransform = FTransform::Identity);
+
+	/** Appends Brush component */
+	bool AppendComponent(UBrushComponent* InComponent, const FTransform& InTransform = FTransform::Identity);
+
+	/** Appends Skeletal Mesh component */
+	bool AppendComponent(const USkeletalMeshComponent* InComponent, const FTransform& InTransform = FTransform::Identity);
+
+	/** Appends Dynamic Mesh component */
+	bool AppendComponent(UDynamicMeshComponent* InComponent, const FTransform& InTransform = FTransform::Identity);
+
+	/** Appends Instanced Static Mesh component */
+	bool AppendComponent(UInstancedStaticMeshComponent* InComponent, const FTransform& InTransform = FTransform::Identity);
+
+	/** Appends Spline Mesh component */
+	bool AppendComponent(USplineMeshComponent* InComponent, const FTransform& InTransform = FTransform::Identity);
+
+	/** Appends Niagara component */
+	bool AppendComponent(UNiagaraComponent* InComponent, const FTransform& InTransform = FTransform::Identity);
+
+	/** Appends a dynamic mesh */
+	bool AppendMesh(const UDynamicMesh* InMesh, const TArray<TWeakObjectPtr<UMaterialInterface>>& InMaterials, const FTransform& InTransform = FTransform::Identity);
+
+	/** Appends a static mesh */
+	bool AppendMesh(UStaticMesh* InMesh, const TArray<TWeakObjectPtr<UMaterialInterface>>& InMaterials, const FTransform& InTransform = FTransform::Identity);
+
+	/** Builds a dynamic mesh by merging all the mesh data imported */
+	bool BuildDynamicMesh(UDynamicMesh* OutMesh, TArray<TWeakObjectPtr<UMaterialInterface>>& OutMaterials, const FCEMeshBuilderParams& InParams = DefaultParams);
+
+	/** Builds a static mesh by merging all the mesh data imported */
+	bool BuildStaticMesh(UStaticMesh* OutMesh, TArray<TWeakObjectPtr<UMaterialInterface>>& OutMaterials, const FCEMeshBuilderParams& InParams = DefaultParams);
+
+	/** Builds a static mesh for the specific instance index */
+	bool BuildStaticMesh(int32 InInstanceIndex, UStaticMesh* OutMesh, FCEMeshInstanceData& OutMeshInstance);
+
+	/** Builds a dynamic mesh for the specific instance index */
+	bool BuildDynamicMesh(int32 InInstanceIndex, UDynamicMesh* OutMesh, FCEMeshInstanceData& OutMeshInstance);
+
+	/** Builds a static mesh for the specific mesh index */
+	bool BuildStaticMesh(uint32 InMeshIndex, UStaticMesh* OutMesh, TArray<FCEMeshInstanceData>& OutMeshInstances);
+
+	/** Builds a dynamic mesh for the specific mesh index */
+	bool BuildDynamicMesh(uint32 InMeshIndex, UDynamicMesh* OutMesh, TArray<FCEMeshInstanceData>& OutMeshInstances);
+
+private:
+	static bool DynamicMeshToStaticMesh(UDynamicMesh* InMesh, UStaticMesh* OutMesh, const TArray<TWeakObjectPtr<UMaterialInterface>>& InMaterials);
+
+	struct FCEMeshInstance
+	{
+		/** Index of the mesh to use for this instance */
+		uint32 MeshIndex;
+
+		/** Data linked to this mesh instance */
+		FCEMeshInstanceData MeshData;
+	};
+
+	FCEMeshInstance* AddMeshInstance(uint32 InMeshIndex, const FTransform& InTransform, const TArray<TWeakObjectPtr<UMaterialInterface>>& InMaterials, TFunctionRef<bool(UE::Geometry::FDynamicMesh3&)> InCreateMeshFunction);
+
+	bool AppendPrimitiveComponent(const UObject* InMeshObject, UPrimitiveComponent* InComponent, const FTransform& InTransform);
+
+	void ClearOutputMesh() const;
+
+	TMap<uint32, UE::Geometry::FDynamicMesh3> Meshes;
+
+	TArray<FCEMeshInstance> MeshInstances;
+
+	UPROPERTY()
+	TObjectPtr<UDynamicMesh> OutputDynamicMesh;
+};
