MotionDesign : ClonerEffector
- Async loading layout and compiling shaders on the task graph to set layout faster

diff --git a/Engine/Plugins/Experimental/ClonerEffector/Source/ClonerEffector/Private/Cloner/Layouts/CEClonerLayoutBase.cpp b/Engine/Plugins/Experimental/ClonerEffector/Source/ClonerEffector/Private/Cloner/Layouts/CEClonerLayoutBase.cpp
--- a/Engine/Plugins/Experimental/ClonerEffector/Source/ClonerEffector/Private/Cloner/Layouts/CEClonerLayoutBase.cpp
+++ b/Engine/Plugins/Experimental/ClonerEffector/Source/ClonerEffector/Private/Cloner/Layouts/CEClonerLayoutBase.cpp
@@ -2,11 +2,13 @@
 
 #include "Cloner/Layouts/CEClonerLayoutBase.h"
 
+#include "Async/Async.h"
 #include "Cloner/CEClonerActor.h"
 #include "Cloner/CEClonerComponent.h"
 #include "Cloner/Extensions/CEClonerExtensionBase.h"
 #include "Engine/Level.h"
 #include "Engine/World.h"
+#include "Misc/Guid.h"
 #include "Misc/PackageName.h"
 #include "NiagaraEmitter.h"
 #include "NiagaraMeshRendererProperties.h"
@@ -116,19 +118,30 @@ void UCEClonerLayoutBase::LoadLayout()
 	FPackagePath::TryFromMountedName(MountedPath, LayoutPackagePath);
 
 	FPackagePath CustomPackagePath;
-	FPackagePath::TryFromPackageName(TEXT("/Game/Temp/") + GetLayoutName().ToString() + TEXT("_") + FString::FromInt(ClonerComponent->GetUniqueID()), CustomPackagePath);
+	FPackagePath::TryFromPackageName(TEXT("/Game/Temp/") + GetLayoutName().ToString() + TEXT("_") + FGuid::NewGuid().ToString(), CustomPackagePath);
 	CustomPackagePath.SetHeaderExtension(EPackageExtension::Asset);
 
-	FLoadPackageAsyncOptionalParams Params;
-	Params.PackagePriority = INT32_MAX;
-	Params.CustomPackageName = CustomPackagePath.GetPackageFName();
-	Params.CompletionDelegate = MakeUnique<FLoadPackageAsyncDelegate>(FLoadPackageAsyncDelegate::CreateUObject(this, &UCEClonerLayoutBase::OnSystemPackageLoaded));
+	UE_LOG(LogCEClonerLayoutBase, Verbose, TEXT("%s : Cloner layout load requested %s - Template system %s - Package %s"), *GetClonerActor()->GetActorNameOrLabel(), *LayoutName.ToString(), *LayoutAssetPath, *CustomPackagePath.GetPackageFName().ToString())
 
-	UE_LOG(LogCEClonerLayoutBase, Verbose, TEXT("%s : Cloner layout load requested %s - Template system %s - Package %s"), *GetClonerActor()->GetActorNameOrLabel(), *LayoutName.ToString(), *LayoutAssetPath, *Params.CustomPackageName.ToString())
+	TWeakObjectPtr<UCEClonerLayoutBase> ThisWeak(this);
+	Async(EAsyncExecution::ThreadPool, [ThisWeak, CustomPackagePath, LayoutPackagePath]()
+	{
+		UCEClonerLayoutBase* This = ThisWeak.Get();
+
+		if (!This)
+		{
+			return;
+		}
+
+		FLoadPackageAsyncOptionalParams Params;
+		Params.PackagePriority = INT32_MAX;
+		Params.CustomPackageName = CustomPackagePath.GetPackageFName();
+		Params.CompletionDelegate = MakeUnique<FLoadPackageAsyncDelegate>(FLoadPackageAsyncDelegate::CreateUObject(This, &UCEClonerLayoutBase::OnSystemPackageLoaded));
 
-	LoadRequestIdentifier = LoadPackageAsync(LayoutPackagePath, MoveTemp(Params));
+		This->LoadRequestIdentifier = LoadPackageAsync(LayoutPackagePath, MoveTemp(Params));
 
-	BindCleanupDelegates();
+		This->BindCleanupDelegates();
+	});
 }
 
 bool UCEClonerLayoutBase::UnloadLayout()
@@ -299,36 +312,61 @@ void UCEClonerLayoutBase::OnSystemPackageLoaded(const FName& InName, UPackage* I
 		NiagaraSystem->RemoveFromRoot();
 		NiagaraSystem->ClearFlags(RF_Standalone);
 
-		for (FNiagaraEmitterHandle& SystemEmitterHandle : NiagaraSystem->GetEmitterHandles())
+		TWeakObjectPtr<UCEClonerLayoutBase> ThisWeak(this);
+		Async(EAsyncExecution::TaskGraphMainThread, [ThisWeak, InName]()
 		{
-			if (const FVersionedNiagaraEmitterData* EmitterData = SystemEmitterHandle.GetEmitterData())
+			UCEClonerLayoutBase* This = ThisWeak.Get();
+
+			if (!This)
+			{
+				return;
+			}
+
+			This->NiagaraSystem->EnsureFullyLoaded();
+
+#if WITH_EDITOR
+			This->NiagaraSystem->WaitForCompilationComplete(/** GPUShaders */true, /** ShowProgress */false);
+#endif
+
+			for (FNiagaraEmitterHandle& SystemEmitterHandle : This->NiagaraSystem->GetEmitterHandles())
 			{
-				for (UNiagaraRendererProperties* EmitterRenderer : EmitterData->GetRenderers())
+				if (const FVersionedNiagaraEmitterData* EmitterData = SystemEmitterHandle.GetEmitterData())
 				{
-					if (UNiagaraMeshRendererProperties* EmitterMeshRenderer = Cast<UNiagaraMeshRendererProperties>(EmitterRenderer))
+					for (UNiagaraRendererProperties* EmitterRenderer : EmitterData->GetRenderers())
 					{
-						EmitterMeshRenderer->Meshes.Empty();
-#if WITH_EDITORONLY_DATA
-						EmitterMeshRenderer->OnMeshChanged();
+						if (UNiagaraMeshRendererProperties* EmitterMeshRenderer = Cast<UNiagaraMeshRendererProperties>(EmitterRenderer))
+						{
+							EmitterMeshRenderer->Meshes.Empty();
+#if WITH_EDITOR
+							EmitterMeshRenderer->OnMeshChanged();
 #endif
 
-						MeshRenderer = EmitterMeshRenderer;
+							This->MeshRenderer = EmitterMeshRenderer;
 
-						UE_LOG(LogCEClonerLayoutBase, Verbose, TEXT("%s : Cloner layout loaded %s - Template system %s - Package %s"), *GetClonerActor()->GetActorNameOrLabel(), *LayoutName.ToString(), *LayoutAssetPath, *InName.ToString())
-
-						OnLayoutLoaded();
+							UE_LOG(LogCEClonerLayoutBase, Verbose, TEXT("%s : Cloner layout loaded %s - Template system %s - Package %s"), *This->GetClonerActor()->GetActorNameOrLabel(), *This->LayoutName.ToString(), *This->LayoutAssetPath, *InName.ToString())
+						}
 					}
 				}
 			}
-		}
+
+			const bool bLayoutLoaded = This->IsLayoutLoaded();
+
+			if (bLayoutLoaded)
+			{
+				This->OnLayoutLoaded();
+			}
+
+			This->OnClonerLayoutLoadedDelegate.Broadcast(This, bLayoutLoaded);
+			This->OnClonerLayoutLoadedDelegate.Clear();
+		});
 	}
 	else
 	{
 		UE_LOG(LogCEClonerLayoutBase, Warning, TEXT("%s : Cloner layout load failed %s - Template system %s - Package %s"), *GetClonerActor()->GetActorNameOrLabel(), *LayoutName.ToString(), *LayoutAssetPath, *InName.ToString())
-	}
 
-	OnClonerLayoutLoadedDelegate.Broadcast(this, IsLayoutLoaded());
-	OnClonerLayoutLoadedDelegate.Clear();
+		OnClonerLayoutLoadedDelegate.Broadcast(this, IsLayoutLoaded());
+		OnClonerLayoutLoadedDelegate.Clear();
+	}
 }
 
 void UCEClonerLayoutBase::BindCleanupDelegates()
