MotionDesign : ClonerEffector
- Fixed cloner mesh builder crash due to invalid mesh ptr, checking if mesh index is valid

diff --git a/Engine/Plugins/Experimental/ClonerEffector/Source/ClonerEffectorMeshBuilder/Private/CEMeshBuilder.cpp b/Engine/Plugins/Experimental/ClonerEffector/Source/ClonerEffectorMeshBuilder/Private/CEMeshBuilder.cpp
--- a/Engine/Plugins/Experimental/ClonerEffector/Source/ClonerEffectorMeshBuilder/Private/CEMeshBuilder.cpp
+++ b/Engine/Plugins/Experimental/ClonerEffector/Source/ClonerEffectorMeshBuilder/Private/CEMeshBuilder.cpp
@@ -726,6 +726,11 @@ bool FCEMeshBuilder::BuildDynamicMesh(UDynamicMesh* OutMesh, TArray<TWeakObjectP
 		// Convert meshes
 		for (const FCEMeshInstance& MeshInstance : MeshInstances)
 		{
+			if (!Meshes.Contains(MeshInstance.MeshIndex))
+			{
+				continue;
+			}
+
 			FDynamicMesh3& Mesh = Meshes[MeshInstance.MeshIndex];
 
 			if (Mesh.TriangleCount() == 0)
@@ -815,6 +820,12 @@ bool FCEMeshBuilder::BuildStaticMesh(int32 InInstanceIndex, UStaticMesh* OutMesh
 	}
 
 	FCEMeshInstance& MeshInstance = MeshInstances[InInstanceIndex];
+
+	if (!Meshes.Contains(MeshInstance.MeshIndex))
+	{
+		return false;
+	}
+
 	FDynamicMesh3& Mesh = Meshes[MeshInstance.MeshIndex];
 
 	if (Mesh.TriangleCount() == 0)
@@ -840,6 +851,12 @@ bool FCEMeshBuilder::BuildDynamicMesh(int32 InInstanceIndex, UDynamicMesh* OutMe
 	}
 
 	FCEMeshInstance& MeshInstance = MeshInstances[InInstanceIndex];
+
+	if (!Meshes.Contains(MeshInstance.MeshIndex))
+	{
+		return false;
+	}
+
 	FDynamicMesh3& Mesh = Meshes[MeshInstance.MeshIndex];
 
 	if (Mesh.TriangleCount() == 0)
