MotionDesign : ClonerEffector
- Changing display name of some cloner/effector properties (LayoutName => Layout, TypeName => Shape, ModeName => Mode)

diff --git a/Engine/Plugins/Experimental/ClonerEffector/Source/ClonerEffector/Public/Cloner/CEClonerComponent.h b/Engine/Plugins/Experimental/ClonerEffector/Source/ClonerEffector/Public/Cloner/CEClonerComponent.h
--- a/Engine/Plugins/Experimental/ClonerEffector/Source/ClonerEffector/Public/Cloner/CEClonerComponent.h
+++ b/Engine/Plugins/Experimental/ClonerEffector/Source/ClonerEffector/Public/Cloner/CEClonerComponent.h
@@ -305,7 +305,7 @@ protected:
 	FLinearColor Color = FLinearColor::White;
 
 	/** Name of the layout to use */
-	UPROPERTY(EditInstanceOnly, Setter, Getter, Category="Layout", meta=(GetOptions="GetClonerLayoutNames"))
+	UPROPERTY(EditInstanceOnly, Setter, Getter, DisplayName="Layout", Category="Layout", meta=(GetOptions="GetClonerLayoutNames"))
 	FName LayoutName = NAME_None;
 
 	/** Active layout used */
diff --git a/Engine/Plugins/Experimental/ClonerEffector/Source/ClonerEffector/Public/Effector/CEEffectorComponent.h b/Engine/Plugins/Experimental/ClonerEffector/Source/ClonerEffector/Public/Effector/CEEffectorComponent.h
--- a/Engine/Plugins/Experimental/ClonerEffector/Source/ClonerEffector/Public/Effector/CEEffectorComponent.h
+++ b/Engine/Plugins/Experimental/ClonerEffector/Source/ClonerEffector/Public/Effector/CEEffectorComponent.h
@@ -210,7 +210,7 @@ protected:
 	FLinearColor Color = FLinearColor::Red;
 
 	/** Name of the shape type to use */
-	UPROPERTY(EditInstanceOnly, Setter, Getter, Category="Shape", meta=(GetOptions="GetEffectorTypeNames"))
+	UPROPERTY(EditInstanceOnly, Setter, Getter, DisplayName="Shape", Category="Shape", meta=(GetOptions="GetEffectorTypeNames"))
 	FName TypeName = NAME_None;
 
 	/** Cached active type used for faster access */
@@ -218,7 +218,7 @@ protected:
 	TObjectPtr<UCEEffectorTypeBase> ActiveType;
 
 	/** Name of the shape type to use */
-	UPROPERTY(EditInstanceOnly, Setter, Getter, Category="Mode", meta=(GetOptions="GetEffectorModeNames"))
+	UPROPERTY(EditInstanceOnly, Setter, Getter, DisplayName="Mode", Category="Mode", meta=(GetOptions="GetEffectorModeNames"))
 	FName ModeName = NAME_None;
 
 	/** Cached active mode used for faster access */
