MotionDesign : ClonerEffector
- Undo after placing a Cloner does not remove it from the level outliner

diff --git a/Engine/Plugins/Experimental/ClonerEffector/Source/ClonerEffector/Private/Cloner/CEClonerActor.cpp b/Engine/Plugins/Experimental/ClonerEffector/Source/ClonerEffector/Private/Cloner/CEClonerActor.cpp
--- a/Engine/Plugins/Experimental/ClonerEffector/Source/ClonerEffector/Private/Cloner/CEClonerActor.cpp
+++ b/Engine/Plugins/Experimental/ClonerEffector/Source/ClonerEffector/Private/Cloner/CEClonerActor.cpp
@@ -70,8 +70,6 @@ ACEClonerActor::ACEClonerActor()
 			GEditor->GetSelectedActors()->SelectionChangedEvent.AddUObject(this, &ACEClonerActor::OnEditorSelectionChanged);
 		}
 #endif
-
-		UCEClonerComponent::OnClonerInitialized().AddUObject(this, &ACEClonerActor::OnClonerInitialized);
 	}
 }
 
@@ -109,18 +107,12 @@ void ACEClonerActor::PostActorCreated()
 
 #if WITH_EDITOR
 	bSpawnDefaultActorAttached = true;
-#endif
-}
 
-void ACEClonerActor::OnClonerInitialized(UCEClonerComponent* InClonerComponent)
-{
-	if (ClonerComponent == InClonerComponent)
+	if (ClonerComponent->bClonerInitialized)
 	{
-#if WITH_EDITOR
-		// PostActorCreated is sometimes called after the cloner is initialized
-		FTSTicker::GetCoreTicker().AddTicker(FTickerDelegate::CreateUObject(this, &ACEClonerActor::SpawnDefaultActorAttached));
-#endif
+		SpawnDefaultActorAttached();
 	}
+#endif
 }
 
 void ACEClonerActor::MigrateDeprecatedProperties()
@@ -369,15 +361,13 @@ void ACEClonerActor::MigrateDeprecatedProperties()
 }
 
 #if WITH_EDITOR
-bool ACEClonerActor::SpawnDefaultActorAttached(float)
+void ACEClonerActor::SpawnDefaultActorAttached()
 {
 	if (bSpawnDefaultActorAttached)
 	{
 		bSpawnDefaultActorAttached = false;
 		ClonerComponent->CreateDefaultActorAttached();
 	}
-
-	return false;
 }
 
 void ACEClonerActor::OnEditorSelectionChanged(UObject* InSelection)
diff --git a/Engine/Plugins/Experimental/ClonerEffector/Source/ClonerEffector/Private/Cloner/CEClonerComponent.cpp b/Engine/Plugins/Experimental/ClonerEffector/Source/ClonerEffector/Private/Cloner/CEClonerComponent.cpp
--- a/Engine/Plugins/Experimental/ClonerEffector/Source/ClonerEffector/Private/Cloner/CEClonerComponent.cpp
+++ b/Engine/Plugins/Experimental/ClonerEffector/Source/ClonerEffector/Private/Cloner/CEClonerComponent.cpp
@@ -26,6 +26,7 @@
 #include "Materials/Material.h"
 #include "Misc/MessageDialog.h"
 #include "Misc/ScopedSlowTask.h"
+#include "ScopedTransaction.h"
 #include "Widgets/Notifications/SNotificationList.h"
 #endif
 
@@ -1229,10 +1230,14 @@ void UCEClonerComponent::CreateDefaultActorAttached()
 		return;
 	}
 
+	FScopedTransaction Transaction(LOCTEXT("CreateDefaultActorAttached", "Create cloner default actor attached"), !GIsTransacting);
+
+	Modify();
+
 	// Spawn attached actor with same flags as this actor
 	FActorSpawnParameters SpawnParameters;
 	SpawnParameters.Owner = Owner;
-	SpawnParameters.ObjectFlags = GetFlags();
+	SpawnParameters.ObjectFlags = GetFlags() | RF_Transactional;
 	SpawnParameters.bTemporaryEditorActor = false;
 
 	const FVector ClonerLocation = GetComponentLocation();
diff --git a/Engine/Plugins/Experimental/ClonerEffector/Source/ClonerEffector/Public/Cloner/CEClonerActor.h b/Engine/Plugins/Experimental/ClonerEffector/Source/ClonerEffector/Public/Cloner/CEClonerActor.h
--- a/Engine/Plugins/Experimental/ClonerEffector/Source/ClonerEffector/Public/Cloner/CEClonerActor.h
+++ b/Engine/Plugins/Experimental/ClonerEffector/Source/ClonerEffector/Public/Cloner/CEClonerActor.h
@@ -50,12 +50,10 @@ protected:
 	//~ End AActor
 
 private:
-	void OnClonerInitialized(UCEClonerComponent* InClonerComponent);
-
 	void MigrateDeprecatedProperties();
 
 #if WITH_EDITOR
-	bool SpawnDefaultActorAttached(float);
+	void SpawnDefaultActorAttached();
 
 	/** Hide outline selection when this cloner only is selected in viewport */
 	void OnEditorSelectionChanged(UObject* InSelection);
