MotionDesign : ClonerEffector
- Fixed previous selected mesh actor is still referenced in data interface and gives warning on deletion

diff --git a/Engine/Plugins/Experimental/ClonerEffector/Source/ClonerEffector/Private/Cloner/Layouts/CEClonerMeshLayout.cpp b/Engine/Plugins/Experimental/ClonerEffector/Source/ClonerEffector/Private/Cloner/Layouts/CEClonerMeshLayout.cpp
--- a/Engine/Plugins/Experimental/ClonerEffector/Source/ClonerEffector/Private/Cloner/Layouts/CEClonerMeshLayout.cpp
+++ b/Engine/Plugins/Experimental/ClonerEffector/Source/ClonerEffector/Private/Cloner/Layouts/CEClonerMeshLayout.cpp
@@ -119,6 +119,11 @@ void UCEClonerMeshLayout::OnLayoutParametersChanged(UCEClonerComponent* InCompon
 	if (USceneComponent* SceneComponent = SceneComponentWeak.Get())
 	{
 		SceneComponent->TransformUpdated.RemoveAll(this);
+
+		if (AActor* Actor = SceneComponent->GetOwner())
+		{
+			Actor->OnDestroyed.RemoveAll(this);
+		}
 	}
 	SceneComponentWeak = nullptr;
 
@@ -127,6 +132,8 @@ void UCEClonerMeshLayout::OnLayoutParametersChanged(UCEClonerComponent* InCompon
 	if (SampleActor && SampleActor->GetRootComponent())
 	{
 		ActorMeshDI->SourceActor = SampleActor;
+		SampleActor->OnDestroyed.AddUniqueDynamic(this, &UCEClonerMeshLayout::OnSampleActorDestroyed);
+
 		SceneComponentWeak = SampleActor->GetRootComponent();
 		SceneComponentWeak->TransformUpdated.AddUObject(this, &UCEClonerMeshLayout::OnSampleMeshTransformed);
 	}
@@ -136,30 +143,30 @@ void UCEClonerMeshLayout::OnLayoutParametersChanged(UCEClonerComponent* InCompon
 		SceneComponentWeak.Reset();
 	}
 
-	if (Asset == ECEClonerMeshAsset::StaticMesh)
+	static const FNiagaraVariable SampleMeshStaticVar(FNiagaraTypeDefinition(UNiagaraDataInterfaceStaticMesh::StaticClass()), TEXT("SampleMeshStatic"));
+	UNiagaraDataInterfaceStaticMesh* StaticMeshDI = Cast<UNiagaraDataInterfaceStaticMesh>(ExposedParameters.GetDataInterface(SampleMeshStaticVar));
+	UStaticMeshComponent* StaticMeshComponent = SampleActor ? SampleActor->FindComponentByClass<UStaticMeshComponent>() : nullptr;
+
+	if (Asset == ECEClonerMeshAsset::StaticMesh && StaticMeshComponent)
 	{
-		static const FNiagaraVariable SampleMeshStaticVar(FNiagaraTypeDefinition(UNiagaraDataInterfaceStaticMesh::StaticClass()), TEXT("SampleMeshStatic"));
-		UNiagaraDataInterfaceStaticMesh* StaticMeshDI = Cast<UNiagaraDataInterfaceStaticMesh>(ExposedParameters.GetDataInterface(SampleMeshStaticVar));
-		if (SampleActorWeak.IsValid())
-		{
-			if (UStaticMeshComponent* StaticMeshComponent = Cast<UStaticMeshComponent>(SampleActorWeak->GetComponentByClass(UStaticMeshComponent::StaticClass())))
-			{
-				StaticMeshDI->SetSourceComponentFromBlueprints(StaticMeshComponent);
-			}
-		}
+		StaticMeshDI->SetSourceComponentFromBlueprints(StaticMeshComponent);
 	}
+	else
+	{
+		StaticMeshDI->SetSourceComponentFromBlueprints(nullptr);
+	}
+
+	static const FNiagaraVariable SampleMeshSkeletalVar(FNiagaraTypeDefinition(UNiagaraDataInterfaceSkeletalMesh::StaticClass()), TEXT("SampleMeshSkeletal"));
+	UNiagaraDataInterfaceSkeletalMesh* SkeletalMeshDI = Cast<UNiagaraDataInterfaceSkeletalMesh>(ExposedParameters.GetDataInterface(SampleMeshSkeletalVar));
+	USkeletalMeshComponent* SkeletalMeshComponent = SampleActor ? SampleActor->FindComponentByClass<USkeletalMeshComponent>() : nullptr;
 
-	if (Asset == ECEClonerMeshAsset::SkeletalMesh)
+	if (Asset == ECEClonerMeshAsset::SkeletalMesh && SkeletalMeshComponent)
 	{
-		static const FNiagaraVariable SampleMeshSkeletalVar(FNiagaraTypeDefinition(UNiagaraDataInterfaceSkeletalMesh::StaticClass()), TEXT("SampleMeshSkeletal"));
-		UNiagaraDataInterfaceSkeletalMesh* SkeletalMeshDI = Cast<UNiagaraDataInterfaceSkeletalMesh>(ExposedParameters.GetDataInterface(SampleMeshSkeletalVar));
-		if (SampleActorWeak.IsValid())
-		{
-			if (USkeletalMeshComponent* SkeletalMeshComponent = Cast<USkeletalMeshComponent>(SampleActorWeak->GetComponentByClass(USkeletalMeshComponent::StaticClass())))
-			{
-				SkeletalMeshDI->SetSourceComponentFromBlueprints(SkeletalMeshComponent);
-			}
-		}
+		SkeletalMeshDI->SetSourceComponentFromBlueprints(SkeletalMeshComponent);
+	}
+	else
+	{
+		SkeletalMeshDI->SetSourceComponentFromBlueprints(nullptr);
 	}
 }
 
@@ -167,3 +174,33 @@ void UCEClonerMeshLayout::OnSampleMeshTransformed(USceneComponent* InComponent,
 {
 	MarkLayoutDirty();
 }
+
+void UCEClonerMeshLayout::OnSampleActorDestroyed(AActor* InDestroyedActor)
+{
+	UCEClonerComponent* ClonerComponent = GetClonerComponent();
+
+	if (!ClonerComponent)
+	{
+		return;
+	}
+
+	const FNiagaraUserRedirectionParameterStore& ExposedParameters = ClonerComponent->GetOverrideParameters();
+
+	static const FNiagaraVariable SampleMeshStaticVar(FNiagaraTypeDefinition(UNiagaraDataInterfaceStaticMesh::StaticClass()), TEXT("SampleMeshStatic"));
+	if (UNiagaraDataInterfaceStaticMesh* StaticMeshDI = Cast<UNiagaraDataInterfaceStaticMesh>(ExposedParameters.GetDataInterface(SampleMeshStaticVar)))
+	{
+		StaticMeshDI->Modify();
+
+		StaticMeshDI->SetSourceComponentFromBlueprints(nullptr);
+	}
+
+	static const FNiagaraVariable SampleMeshSkeletalVar(FNiagaraTypeDefinition(UNiagaraDataInterfaceSkeletalMesh::StaticClass()), TEXT("SampleMeshSkeletal"));
+	if (UNiagaraDataInterfaceSkeletalMesh* SkeletalMeshDI = Cast<UNiagaraDataInterfaceSkeletalMesh>(ExposedParameters.GetDataInterface(SampleMeshSkeletalVar)))
+	{
+		SkeletalMeshDI->Modify();
+
+		SkeletalMeshDI->SetSourceComponentFromBlueprints(nullptr);
+	}
+
+	MarkLayoutDirty();
+}
diff --git a/Engine/Plugins/Experimental/ClonerEffector/Source/ClonerEffector/Public/Cloner/Layouts/CEClonerMeshLayout.h b/Engine/Plugins/Experimental/ClonerEffector/Source/ClonerEffector/Public/Cloner/Layouts/CEClonerMeshLayout.h
--- a/Engine/Plugins/Experimental/ClonerEffector/Source/ClonerEffector/Public/Cloner/Layouts/CEClonerMeshLayout.h
+++ b/Engine/Plugins/Experimental/ClonerEffector/Source/ClonerEffector/Public/Cloner/Layouts/CEClonerMeshLayout.h
@@ -89,10 +89,13 @@ protected:
 
 	void OnSampleMeshTransformed(USceneComponent* InComponent, EUpdateTransformFlags InFlags, ETeleportType InType);
 
+	UFUNCTION()
+	void OnSampleActorDestroyed(AActor* InDestroyedActor);
+
 	UPROPERTY(EditInstanceOnly, Setter, Getter, Category="Layout")
 	ECEClonerMeshAsset Asset = ECEClonerMeshAsset::StaticMesh;
 
-	UPROPERTY(EditInstanceOnly, Setter, Getter, DisplayName="Sample Actor", Category="Layout")
+	UPROPERTY(EditInstanceOnly, Setter, Getter, DisplayName="Sample Actor", Category="Layout", meta=(DisplayAfter="Asset"))
 	TWeakObjectPtr<AActor> SampleActorWeak;
 
 	UPROPERTY(EditInstanceOnly, Setter, Getter, Category="Layout")
diff --git a/Engine/Plugins/FX/Niagara/Source/Niagara/Private/DataInterface/NiagaraDataInterfaceStaticMesh.cpp b/Engine/Plugins/FX/Niagara/Source/Niagara/Private/DataInterface/NiagaraDataInterfaceStaticMesh.cpp
--- a/Engine/Plugins/FX/Niagara/Source/Niagara/Private/DataInterface/NiagaraDataInterfaceStaticMesh.cpp
+++ b/Engine/Plugins/FX/Niagara/Source/Niagara/Private/DataInterface/NiagaraDataInterfaceStaticMesh.cpp
@@ -3606,7 +3606,7 @@ void UNiagaraDataInterfaceStaticMesh::SetSourceComponentFromBlueprints(UStaticMe
 	++ChangeId;
 	UnbindSourceDelegates();
 	SourceComponent = ComponentToUse;
-	SoftSourceActor = ComponentToUse->GetOwner();
+	SoftSourceActor = ComponentToUse ? ComponentToUse->GetOwner() : nullptr;
 	BindSourceDelegates();
 }
 
diff --git a/Engine/Plugins/FX/Niagara/Source/Niagara/Private/NiagaraDataInterfaceSkeletalMesh.cpp b/Engine/Plugins/FX/Niagara/Source/Niagara/Private/NiagaraDataInterfaceSkeletalMesh.cpp
--- a/Engine/Plugins/FX/Niagara/Source/Niagara/Private/NiagaraDataInterfaceSkeletalMesh.cpp
+++ b/Engine/Plugins/FX/Niagara/Source/Niagara/Private/NiagaraDataInterfaceSkeletalMesh.cpp
@@ -3480,7 +3480,7 @@ void UNiagaraDataInterfaceSkeletalMesh::SetSourceComponentFromBlueprints(USkelet
 	++ChangeId;
 	UnbindSourceDelegates();
 	SourceComponent = ComponentToUse;
-	SoftSourceActor = ComponentToUse->GetOwner();
+	SoftSourceActor = ComponentToUse ? ComponentToUse->GetOwner() : nullptr;
 	BindSourceDelegates();
 }
 
