MotionDesign : ClonerEffector
- Fixed effector previously linked to cloner would be lost after migration to latest version

diff --git a/Engine/Plugins/Experimental/ClonerEffector/Source/ClonerEffector/Private/Cloner/CEClonerActor.cpp b/Engine/Plugins/Experimental/ClonerEffector/Source/ClonerEffector/Private/Cloner/CEClonerActor.cpp
--- a/Engine/Plugins/Experimental/ClonerEffector/Source/ClonerEffector/Private/Cloner/CEClonerActor.cpp
+++ b/Engine/Plugins/Experimental/ClonerEffector/Source/ClonerEffector/Private/Cloner/CEClonerActor.cpp
@@ -214,6 +214,8 @@ void ACEClonerActor::MigrateDeprecatedProperties()
 			{
 				if (ACEEffectorActor* Effector = EffectorWeak.Get())
 				{
+					// Register to channel before linking
+					Effector->RegisterToChannel();
 					EffectorExtension->LinkEffector(Effector);
 				}
 			}
diff --git a/Engine/Plugins/Experimental/ClonerEffector/Source/ClonerEffector/Private/Effector/CEEffectorActor.cpp b/Engine/Plugins/Experimental/ClonerEffector/Source/ClonerEffector/Private/Effector/CEEffectorActor.cpp
--- a/Engine/Plugins/Experimental/ClonerEffector/Source/ClonerEffector/Private/Effector/CEEffectorActor.cpp
+++ b/Engine/Plugins/Experimental/ClonerEffector/Source/ClonerEffector/Private/Effector/CEEffectorActor.cpp
@@ -229,3 +229,11 @@ void ACEEffectorActor::MigrateDeprecatedProperties()
 
 	MigrateToVersion = INDEX_NONE;
 }
+
+void ACEEffectorActor::RegisterToChannel() const
+{
+	if (EffectorComponent)
+	{
+		EffectorComponent->RegisterToChannel();
+	}
+}
diff --git a/Engine/Plugins/Experimental/ClonerEffector/Source/ClonerEffector/Public/Effector/CEEffectorActor.h b/Engine/Plugins/Experimental/ClonerEffector/Source/ClonerEffector/Public/Effector/CEEffectorActor.h
--- a/Engine/Plugins/Experimental/ClonerEffector/Source/ClonerEffector/Public/Effector/CEEffectorActor.h
+++ b/Engine/Plugins/Experimental/ClonerEffector/Source/ClonerEffector/Public/Effector/CEEffectorActor.h
@@ -49,6 +49,7 @@ protected:
 
 private:
 	void MigrateDeprecatedProperties();
+	void RegisterToChannel() const;
 
 	UPROPERTY(VisibleInstanceOnly, Category="Effector")
 	TObjectPtr<UCEEffectorComponent> EffectorComponent;
