MotionDesign : ClonerEffector
- Fixed cloner lifetime scale curve visibility issue

diff --git a/Engine/Plugins/Experimental/ClonerEffector/Source/ClonerEffectorEditor/Private/Cloner/Customizations/CEEditorClonerLifetimeExtensionDetailCustomization.cpp b/Engine/Plugins/Experimental/ClonerEffector/Source/ClonerEffectorEditor/Private/Cloner/Customizations/CEEditorClonerLifetimeExtensionDetailCustomization.cpp
--- a/Engine/Plugins/Experimental/ClonerEffector/Source/ClonerEffectorEditor/Private/Cloner/Customizations/CEEditorClonerLifetimeExtensionDetailCustomization.cpp
+++ b/Engine/Plugins/Experimental/ClonerEffector/Source/ClonerEffectorEditor/Private/Cloner/Customizations/CEEditorClonerLifetimeExtensionDetailCustomization.cpp
@@ -6,6 +6,7 @@
 #include "DetailBuilderTypes.h"
 #include "DetailCategoryBuilder.h"
 #include "DetailLayoutBuilder.h"
+#include "DetailWidgetRow.h"
 #include "GameFramework/Actor.h"
 #include "NiagaraDataInterfaceCurve.h"
 #include "PropertyHandle.h"
@@ -27,15 +28,28 @@ void FCEEditorClonerLifetimeExtensionDetailCustomization::CustomizeDetails(IDeta
 			continue;
 		}
 
+		UNiagaraDataInterfaceCurve* CurveDI = LifetimeExtension->GetLifetimeScaleCurveDI();
+
+		if (!CurveDI)
+		{
+			continue;
+		}
+
 		const FName CategoryName = LifetimeExtension->GetExtensionName();
+		IDetailCategoryBuilder& CategoryBuilder = InDetailBuilder.EditCategory(FName(CategoryName.ToString() + TEXT("Curve")), FText::GetEmpty(), ECategoryPriority::Uncommon);
 
-		IDetailCategoryBuilder& CategoryBuilder = InDetailBuilder.EditCategory(CategoryName);
-		CategoryBuilder.SetShowAdvanced(true);
+		// Hide other properties, only curve will be shown instead of tree
+		for (FProperty* Property : TFieldRange<FProperty>(CurveDI->GetClass()))
+		{
+			if (Property)
+			{
+				Property->SetMetaData(TEXT("EditCondition"), TEXT("false"));
+				Property->SetMetaData(TEXT("EditConditionHides"), TEXT("true"));
+			}
+		}
 
-		/**
-		 * UNiagaraDataInterfaceCurve cannot display simultaneously multiple curves, so we need to add them separately
-		 */
-		if (IDetailPropertyRow* Row = CategoryBuilder.AddExternalObjects({LifetimeExtension->GetLifetimeScaleCurveDI()}, EPropertyLocation::Advanced, Params))
+		// UNiagaraDataInterfaceCurve cannot display simultaneously multiple curves, so we need to add them separately
+		if (IDetailPropertyRow* Row = CategoryBuilder.AddExternalObjects({CurveDI}, EPropertyLocation::Common, Params))
 		{
 			const TAttribute<EVisibility> VisibilityAttr = TAttribute<EVisibility>::CreateSP(this, &FCEEditorClonerLifetimeExtensionDetailCustomization::GetCurveVisibility, LifetimeExtensionWeak);
 			Row->Visibility(VisibilityAttr);
