MotionDesign : ClonerEffector
- Ability to add a niagara lifecycle and cache track from cloner details panel to create a deterministic simulation

diff --git a/Engine/Plugins/Experimental/ClonerEffector/Source/ClonerEffector/Public/Cloner/Extensions/CEClonerExtensionBase.h b/Engine/Plugins/Experimental/ClonerEffector/Source/ClonerEffector/Public/Cloner/Extensions/CEClonerExtensionBase.h
--- a/Engine/Plugins/Experimental/ClonerEffector/Source/ClonerEffector/Public/Cloner/Extensions/CEClonerExtensionBase.h
+++ b/Engine/Plugins/Experimental/ClonerEffector/Source/ClonerEffector/Public/Cloner/Extensions/CEClonerExtensionBase.h
@@ -46,7 +46,7 @@ public:
 #endif
 
 	/** Get the cloner component using this extension */
-	UCEClonerComponent* GetClonerComponent() const;
+	CLONEREFFECTOR_API UCEClonerComponent* GetClonerComponent() const;
 
 	/** Get the cloner component using this extension checked version */
 	UCEClonerComponent* GetClonerComponentChecked() const;
diff --git a/Engine/Plugins/Experimental/ClonerEffector/Source/ClonerEffectorEditor/ClonerEffectorEditor.Build.cs b/Engine/Plugins/Experimental/ClonerEffector/Source/ClonerEffectorEditor/ClonerEffectorEditor.Build.cs
--- a/Engine/Plugins/Experimental/ClonerEffector/Source/ClonerEffectorEditor/ClonerEffectorEditor.Build.cs
+++ b/Engine/Plugins/Experimental/ClonerEffector/Source/ClonerEffectorEditor/ClonerEffectorEditor.Build.cs
@@ -17,7 +17,7 @@ public class ClonerEffectorEditor : ModuleRules
 		        "Engine",
 		        "Slate",
 		        "SlateCore",
-		        "UnrealEd",
+		        "UnrealEd"
 	        }
         );
 
@@ -26,8 +26,13 @@ public class ClonerEffectorEditor : ModuleRules
 	        {
 		        "ClonerEffector",
 		        "InputCore",
+				"MovieScene",
+				"Niagara",
+		        "NiagaraEditor",
+		        "NiagaraSimCaching",
 		        "Projects",
 				"PropertyEditor",
+				"Sequencer",
 				"ToolMenus"
 	        }
         );
diff --git a/Engine/Plugins/Experimental/ClonerEffector/Source/ClonerEffectorEditor/Private/CEEditorModule.cpp b/Engine/Plugins/Experimental/ClonerEffector/Source/ClonerEffectorEditor/Private/CEEditorModule.cpp
--- a/Engine/Plugins/Experimental/ClonerEffector/Source/ClonerEffectorEditor/Private/CEEditorModule.cpp
+++ b/Engine/Plugins/Experimental/ClonerEffector/Source/ClonerEffectorEditor/Private/CEEditorModule.cpp
@@ -14,23 +14,25 @@
 #include "Cloner/Extensions/CEClonerLifetimeExtension.h"
 #include "Cloner/Layouts/CEClonerMeshLayout.h"
 #include "Cloner/Layouts/CEClonerSplineLayout.h"
+#include "Cloner/Sequencer/MovieSceneClonerTrackEditor.h"
 #include "Effector/Customizations/CEEditorEffectorComponentDetailCustomization.h"
 #include "Effector/Customizations/CEEditorEffectorTypeDetailCustomization.h"
 #include "Effector/CEEffectorActor.h"
 #include "Effector/CEEffectorComponent.h"
 #include "Effector/Customizations/CEEditorEffectorActorDetailCustomization.h"
 #include "Effector/Types/CEEffectorBoundType.h"
+#include "ISequencerModule.h"
 #include "Modules/ModuleManager.h"
 #include "PropertyEditorModule.h"
 #include "Styles/CEEditorStyle.h"
 
 void FCEEditorModule::StartupModule()
 {
-	FPropertyEditorModule& PropertyModule = FModuleManager::GetModuleChecked<FPropertyEditorModule>(PropertyEditorName);
-
 	// Load styles
 	FCEEditorStyle::Get();
 
+	FPropertyEditorModule& PropertyModule = FModuleManager::LoadModuleChecked<FPropertyEditorModule>(TEXT("PropertyEditor"));
+
 	// Cloner customization
 	PropertyModule.RegisterCustomClassLayout(ACEClonerActor::StaticClass()->GetFName(), FOnGetDetailCustomizationInstance::CreateStatic(&FCEEditorClonerActorDetailCustomization::MakeInstance));
 	PropertyModule.RegisterCustomClassLayout(UCEClonerComponent::StaticClass()->GetFName(), FOnGetDetailCustomizationInstance::CreateStatic(&FCEEditorClonerComponentDetailCustomization::MakeInstance));
@@ -43,13 +45,17 @@ void FCEEditorModule::StartupModule()
 	PropertyModule.RegisterCustomClassLayout(ACEEffectorActor::StaticClass()->GetFName(), FOnGetDetailCustomizationInstance::CreateStatic(&FCEEditorEffectorActorDetailCustomization::MakeInstance));
 	PropertyModule.RegisterCustomClassLayout(UCEEffectorComponent::StaticClass()->GetFName(), FOnGetDetailCustomizationInstance::CreateStatic(&FCEEditorEffectorComponentDetailCustomization::MakeInstance));
 	PropertyModule.RegisterCustomClassLayout(UCEEffectorBoundType::StaticClass()->GetFName(), FOnGetDetailCustomizationInstance::CreateStatic(&FCEEditorEffectorTypeDetailCustomization::MakeInstance));
+
+	// Custom cloner track
+	ISequencerModule& SequencerModule = FModuleManager::LoadModuleChecked<ISequencerModule>(TEXT("Sequencer"));
+	ClonerTrackCreateEditorHandle = SequencerModule.RegisterTrackEditor(FOnCreateTrackEditor::CreateStatic(&FMovieSceneClonerTrackEditor::CreateTrackEditor));
 }
 
 void FCEEditorModule::ShutdownModule()
 {
-	if (FModuleManager::Get().IsModuleLoaded(PropertyEditorName) && UObjectInitialized())
+	if (FModuleManager::Get().IsModuleLoaded("PropertyEditor") && UObjectInitialized())
 	{
-		FPropertyEditorModule& PropertyModule = FModuleManager::GetModuleChecked<FPropertyEditorModule>(PropertyEditorName);
+		FPropertyEditorModule& PropertyModule = FModuleManager::GetModuleChecked<FPropertyEditorModule>("PropertyEditor");
 
 		// Cloner customization
 		PropertyModule.UnregisterCustomClassLayout(ACEClonerActor::StaticClass()->GetFName());
@@ -64,6 +70,14 @@ void FCEEditorModule::ShutdownModule()
 		PropertyModule.UnregisterCustomClassLayout(UCEEffectorComponent::StaticClass()->GetFName());
 		PropertyModule.UnregisterCustomClassLayout(UCEEffectorBoundType::StaticClass()->GetFName());
 	}
+
+	// Custom cloner track
+	if (FModuleManager::Get().IsModuleLoaded("Sequencer"))
+	{
+		ISequencerModule& SequencerModule = FModuleManager::GetModuleChecked<ISequencerModule>(TEXT("Sequencer"));
+		SequencerModule.UnRegisterTrackEditor(ClonerTrackCreateEditorHandle);
+		ClonerTrackCreateEditorHandle.Reset();
+	}
 }
 
 IMPLEMENT_MODULE(FCEEditorModule, ClonerEffectorEditor)
\ No newline at end of file
diff --git a/Engine/Plugins/Experimental/ClonerEffector/Source/ClonerEffectorEditor/Private/CEEditorModule.h b/Engine/Plugins/Experimental/ClonerEffector/Source/ClonerEffectorEditor/Private/CEEditorModule.h
--- a/Engine/Plugins/Experimental/ClonerEffector/Source/ClonerEffectorEditor/Private/CEEditorModule.h
+++ b/Engine/Plugins/Experimental/ClonerEffector/Source/ClonerEffectorEditor/Private/CEEditorModule.h
@@ -7,10 +7,11 @@
 class FCEEditorModule : public IModuleInterface
 {
 public:
-	static constexpr TCHAR PropertyEditorName[] = TEXT("PropertyEditor");
-
 	//~ Begin IModuleInterface
     virtual void StartupModule() override;
     virtual void ShutdownModule() override;
 	//~ End IModuleInterface
+
+protected:
+	FDelegateHandle ClonerTrackCreateEditorHandle;
 };
diff --git a/Engine/Plugins/Experimental/ClonerEffector/Source/ClonerEffectorEditor/Private/Cloner/Customizations/CEEditorClonerLifetimeExtensionDetailCustomization.cpp b/Engine/Plugins/Experimental/ClonerEffector/Source/ClonerEffectorEditor/Private/Cloner/Customizations/CEEditorClonerLifetimeExtensionDetailCustomization.cpp
--- a/Engine/Plugins/Experimental/ClonerEffector/Source/ClonerEffectorEditor/Private/Cloner/Customizations/CEEditorClonerLifetimeExtensionDetailCustomization.cpp
+++ b/Engine/Plugins/Experimental/ClonerEffector/Source/ClonerEffectorEditor/Private/Cloner/Customizations/CEEditorClonerLifetimeExtensionDetailCustomization.cpp
@@ -3,6 +3,7 @@
 #include "Cloner/Customizations/CEEditorClonerLifetimeExtensionDetailCustomization.h"
 
 #include "Cloner/Extensions/CEClonerLifetimeExtension.h"
+#include "Cloner/Sequencer/MovieSceneClonerTrackEditor.h"
 #include "DetailBuilderTypes.h"
 #include "DetailCategoryBuilder.h"
 #include "DetailLayoutBuilder.h"
@@ -10,6 +11,10 @@
 #include "GameFramework/Actor.h"
 #include "NiagaraDataInterfaceCurve.h"
 #include "PropertyHandle.h"
+#include "Widgets/Input/SButton.h"
+#include "Widgets/SBoxPanel.h"
+
+#define LOCTEXT_NAMESPACE "CEEditorClonerLifetimeExtensionDetailCustomization"
 
 void FCEEditorClonerLifetimeExtensionDetailCustomization::CustomizeDetails(IDetailLayoutBuilder& InDetailBuilder)
 {
@@ -36,7 +41,7 @@ void FCEEditorClonerLifetimeExtensionDetailCustomization::CustomizeDetails(IDeta
 		}
 
 		const FName CategoryName = LifetimeExtension->GetExtensionName();
-		IDetailCategoryBuilder& CategoryBuilder = InDetailBuilder.EditCategory(FName(CategoryName.ToString() + TEXT("Curve")), FText::GetEmpty(), ECategoryPriority::Uncommon);
+		IDetailCategoryBuilder& CurveCategoryBuilder = InDetailBuilder.EditCategory(FName(CategoryName.ToString() + TEXT("Curve")), FText::GetEmpty(), ECategoryPriority::Uncommon);
 
 		// Hide other properties, only curve will be shown instead of tree
 		for (FProperty* Property : TFieldRange<FProperty>(CurveDI->GetClass()))
@@ -49,11 +54,36 @@ void FCEEditorClonerLifetimeExtensionDetailCustomization::CustomizeDetails(IDeta
 		}
 
 		// UNiagaraDataInterfaceCurve cannot display simultaneously multiple curves, so we need to add them separately
-		if (IDetailPropertyRow* Row = CategoryBuilder.AddExternalObjects({CurveDI}, EPropertyLocation::Common, Params))
+		if (IDetailPropertyRow* Row = CurveCategoryBuilder.AddExternalObjects({CurveDI}, EPropertyLocation::Common, Params))
 		{
 			const TAttribute<EVisibility> VisibilityAttr = TAttribute<EVisibility>::CreateSP(this, &FCEEditorClonerLifetimeExtensionDetailCustomization::GetCurveVisibility, LifetimeExtensionWeak);
 			Row->Visibility(VisibilityAttr);
 		}
+
+		if (LifetimeExtensionsWeak.Num() == 1)
+		{
+			const FText ButtonLabel = LOCTEXT("AddSequencerTrack", "Add Sequencer Tracks");
+
+			IDetailCategoryBuilder& LifetimeCategoryBuilder = InDetailBuilder.EditCategory(CategoryName);
+
+			LifetimeCategoryBuilder.AddCustomRow(ButtonLabel)
+			.WholeRowContent()
+			.HAlign(HAlign_Left)
+			[
+				SNew(SVerticalBox)
+				+ SVerticalBox::Slot()
+				.Padding(0.f, 3.f)
+				.AutoHeight()
+				[
+					SNew(SButton)
+					.Text(ButtonLabel)
+					.HAlign(HAlign_Center)
+					.VAlign(VAlign_Fill)
+					.IsEnabled(this, &FCEEditorClonerLifetimeExtensionDetailCustomization::CanAddSequencerTracks, LifetimeExtensionWeak)
+					.OnClicked(this, &FCEEditorClonerLifetimeExtensionDetailCustomization::OnAddSequencerTracks, LifetimeExtensionWeak)
+				]
+			];
+		}
 	}
 }
 
@@ -71,3 +101,29 @@ EVisibility FCEEditorClonerLifetimeExtensionDetailCustomization::GetCurveVisibil
 		? EVisibility::Visible
 		: EVisibility::Collapsed;
 }
+
+bool FCEEditorClonerLifetimeExtensionDetailCustomization::CanAddSequencerTracks(TWeakObjectPtr<UCEClonerLifetimeExtension> InClonerLifetimeExtension) const
+{
+	uint32 TrackCount = 0;
+
+	if (const UCEClonerLifetimeExtension* LifetimeExtension = InClonerLifetimeExtension.Get())
+	{
+		FMovieSceneClonerTrackEditor::OnClonerTrackExists.Broadcast(LifetimeExtension->GetClonerComponent(), TrackCount);
+	}
+
+	// Lifecycle + cache tracks
+	constexpr uint32 ExpectedTrackCount = 2;
+	return TrackCount < ExpectedTrackCount;
+}
+
+FReply FCEEditorClonerLifetimeExtensionDetailCustomization::OnAddSequencerTracks(TWeakObjectPtr<UCEClonerLifetimeExtension> InClonerLifetimeExtension)
+{
+	if (const UCEClonerLifetimeExtension* LifetimeExtension = InClonerLifetimeExtension.Get())
+	{
+		FMovieSceneClonerTrackEditor::OnAddClonerTrack.Broadcast(LifetimeExtension->GetClonerComponent());
+	}
+
+	return FReply::Handled();
+}
+
+#undef LOCTEXT_NAMESPACE
diff --git a/Engine/Plugins/Experimental/ClonerEffector/Source/ClonerEffectorEditor/Private/Cloner/Customizations/CEEditorClonerLifetimeExtensionDetailCustomization.h b/Engine/Plugins/Experimental/ClonerEffector/Source/ClonerEffectorEditor/Private/Cloner/Customizations/CEEditorClonerLifetimeExtensionDetailCustomization.h
--- a/Engine/Plugins/Experimental/ClonerEffector/Source/ClonerEffectorEditor/Private/Cloner/Customizations/CEEditorClonerLifetimeExtensionDetailCustomization.h
+++ b/Engine/Plugins/Experimental/ClonerEffector/Source/ClonerEffectorEditor/Private/Cloner/Customizations/CEEditorClonerLifetimeExtensionDetailCustomization.h
@@ -4,7 +4,9 @@
 
 #include "IDetailCustomization.h"
 #include "Templates/SharedPointer.h"
+#include "UObject/WeakObjectPtrTemplatesFwd.h"
 
+class FReply;
 class IDetailCategoryBuilder;
 class IPropertyHandle;
 class UCEClonerLifetimeExtension;
@@ -26,6 +28,9 @@ public:
 protected:
 	EVisibility GetCurveVisibility(TWeakObjectPtr<UCEClonerLifetimeExtension> InExtensionWeak) const;
 
+	bool CanAddSequencerTracks(TWeakObjectPtr<UCEClonerLifetimeExtension> InClonerLifetimeExtension) const;
+	FReply OnAddSequencerTracks(TWeakObjectPtr<UCEClonerLifetimeExtension> InClonerLifetimeExtension);
+
 	TSharedPtr<IPropertyHandle> LifetimeEnabledPropertyHandle;
 	TSharedPtr<IPropertyHandle> LifetimeScaleEnabledPropertyHandle;
 };
diff --git a/Engine/Plugins/Experimental/ClonerEffector/Source/ClonerEffectorEditor/Private/Cloner/Sequencer/MovieSceneClonerTrackEditor.cpp b/Engine/Plugins/Experimental/ClonerEffector/Source/ClonerEffectorEditor/Private/Cloner/Sequencer/MovieSceneClonerTrackEditor.cpp
new file mode 100644
--- /dev/null
+++ b/Engine/Plugins/Experimental/ClonerEffector/Source/ClonerEffectorEditor/Private/Cloner/Sequencer/MovieSceneClonerTrackEditor.cpp
@@ -0,0 +1,135 @@
+// Copyright Epic Games, Inc. All Rights Reserved.
+
+#include "Cloner/Sequencer/MovieSceneClonerTrackEditor.h"
+
+#include "Cloner/CEClonerComponent.h"
+#include "MovieScene/MovieSceneNiagaraSystemSpawnSection.h"
+#include "MovieScene/MovieSceneNiagaraSystemTrack.h"
+#include "Niagara/Sequencer/MovieSceneNiagaraCacheSection.h"
+#include "Niagara/Sequencer/MovieSceneNiagaraCacheTrack.h"
+
+#define LOCTEXT_NAMESPACE "MovieSceneClonerTrackEditor"
+
+FMovieSceneClonerTrackEditor::FOnAddClonerTrack FMovieSceneClonerTrackEditor::OnAddClonerTrack;
+FMovieSceneClonerTrackEditor::FOnClonerTrackExists FMovieSceneClonerTrackEditor::OnClonerTrackExists;
+
+FMovieSceneClonerTrackEditor::~FMovieSceneClonerTrackEditor()
+{
+	OnAddClonerTrack.RemoveAll(this);
+	OnClonerTrackExists.RemoveAll(this);
+}
+
+bool FMovieSceneClonerTrackEditor::SupportsType(TSubclassOf<UMovieSceneTrack> InTrackClass) const
+{
+	return InTrackClass == UMovieSceneNiagaraCacheTrack::StaticClass()
+		|| InTrackClass == UMovieSceneNiagaraSystemTrack::StaticClass();
+}
+
+void FMovieSceneClonerTrackEditor::BindDelegates()
+{
+	OnAddClonerTrack.AddSP(this, &FMovieSceneClonerTrackEditor::ExecuteAddTrack);
+	OnClonerTrackExists.AddSP(this, &FMovieSceneClonerTrackEditor::ExecuteTrackExists);
+}
+
+void FMovieSceneClonerTrackEditor::ExecuteAddTrack(UCEClonerComponent* InCloner)
+{
+	if (!IsValid(InCloner))
+	{
+		return;
+	}
+
+	UMovieScene* FocusedMovieScene = GetFocusedMovieScene();
+	if (!FocusedMovieScene || FocusedMovieScene->IsReadOnly())
+	{
+		return;
+	}
+
+	const TSharedPtr<ISequencer> SequencerPtr = GetSequencer();
+	if (!SequencerPtr)
+	{
+		return;
+	}
+
+	const FScopedTransaction Transaction(LOCTEXT("AddClonerTracks", "Add Cloner Tracks"));
+
+	FocusedMovieScene->Modify();
+
+	const FGuid ComponentBinding = SequencerPtr->GetHandleToObject(InCloner);
+	if (FocusedMovieScene->FindSpawnable(ComponentBinding))
+	{
+		// we only want to add tracks for possessables
+		return;
+	}
+
+	// Add lifecycle track
+	UMovieSceneTrack* LifeCycleTrack = FocusedMovieScene->FindTrack(UMovieSceneNiagaraSystemTrack::StaticClass(), ComponentBinding);
+
+	if (!LifeCycleTrack)
+	{
+		UMovieSceneNiagaraSystemTrack* NiagaraSystemTrack = FocusedMovieScene->AddTrack<UMovieSceneNiagaraSystemTrack>(ComponentBinding);
+		NiagaraSystemTrack->SetDisplayName(LOCTEXT("ClonerLifeCycleTrackName", "Cloner Life Cycle"));
+
+		UMovieSceneNiagaraSystemSpawnSection* NiagaraSpawnSection = Cast<UMovieSceneNiagaraSystemSpawnSection>(NiagaraSystemTrack->CreateNewSection());
+		NiagaraSpawnSection->SetAgeUpdateMode(ENiagaraAgeUpdateMode::DesiredAge);
+
+		NiagaraSpawnSection->SetRange(TRange<FFrameNumber>(
+			FocusedMovieScene->GetPlaybackRange().GetLowerBound(),
+			FocusedMovieScene->GetPlaybackRange().GetUpperBound()
+		));
+
+		NiagaraSystemTrack->AddSection(*NiagaraSpawnSection);
+	}
+
+	// Add cache track
+	UMovieSceneTrack* CacheTrack = FocusedMovieScene->FindTrack(UMovieSceneNiagaraCacheTrack::StaticClass(), ComponentBinding);
+
+	if (!CacheTrack)
+	{
+		UMovieSceneNiagaraCacheTrack* NiagaraCacheTrack = FocusedMovieScene->AddTrack<UMovieSceneNiagaraCacheTrack>(ComponentBinding);
+		NiagaraCacheTrack->SetDisplayName(LOCTEXT("ClonerSimCacheTrackName", "Cloner Sim Cache"));
+
+		UMovieSceneNiagaraCacheSection* NiagaraCacheSection = Cast<UMovieSceneNiagaraCacheSection>(NiagaraCacheTrack->CreateNewSection());
+
+		NiagaraCacheSection->SetRange(TRange<FFrameNumber>(
+			FocusedMovieScene->GetPlaybackRange().GetLowerBound(),
+			FocusedMovieScene->GetPlaybackRange().GetUpperBound()
+		));
+
+		NiagaraCacheTrack->AddSection(*NiagaraCacheSection);
+	}
+
+	SequencerPtr->NotifyMovieSceneDataChanged(EMovieSceneDataChangeType::MovieSceneStructureItemAdded);
+}
+
+void FMovieSceneClonerTrackEditor::ExecuteTrackExists(UCEClonerComponent* InCloner, uint32& OutCount) const
+{
+	const UMovieScene* FocusedMovieScene = GetFocusedMovieScene();
+	if (!FocusedMovieScene)
+	{
+		return;
+	}
+
+	const TSharedPtr<ISequencer> SequencerPtr = GetSequencer();
+	if (!SequencerPtr)
+	{
+		return;
+	}
+
+	const FGuid ComponentBinding = SequencerPtr->GetHandleToObject(InCloner, /** CreateIfMissing */false);
+	if (!ComponentBinding.IsValid())
+	{
+		return;
+	}
+
+	if (UMovieSceneTrack* NiagaraSystemTrack = FocusedMovieScene->FindTrack(UMovieSceneNiagaraSystemTrack::StaticClass(), ComponentBinding))
+	{
+		OutCount++;
+	}
+
+	if (UMovieSceneTrack* NiagaraCacheTrack = FocusedMovieScene->FindTrack(UMovieSceneNiagaraCacheTrack::StaticClass(), ComponentBinding))
+	{
+		OutCount++;
+	}
+}
+
+#undef LOCTEXT_NAMESPACE
diff --git a/Engine/Plugins/Experimental/ClonerEffector/Source/ClonerEffectorEditor/Private/Cloner/Sequencer/MovieSceneClonerTrackEditor.h b/Engine/Plugins/Experimental/ClonerEffector/Source/ClonerEffectorEditor/Private/Cloner/Sequencer/MovieSceneClonerTrackEditor.h
new file mode 100644
--- /dev/null
+++ b/Engine/Plugins/Experimental/ClonerEffector/Source/ClonerEffectorEditor/Private/Cloner/Sequencer/MovieSceneClonerTrackEditor.h
@@ -0,0 +1,42 @@
+// Copyright Epic Games, Inc. All Rights Reserved.
+
+#pragma once
+
+#include "MovieSceneTrackEditor.h"
+
+class UCEClonerComponent;
+class ISequencer;
+
+/** Cloner track editor to add niagara track and section */
+class FMovieSceneClonerTrackEditor : public FMovieSceneTrackEditor
+{
+public:
+	DECLARE_MULTICAST_DELEGATE_OneParam(FOnAddClonerTrack, UCEClonerComponent* /** InCloner */)
+	static FOnAddClonerTrack OnAddClonerTrack;
+
+	DECLARE_MULTICAST_DELEGATE_TwoParams(FOnClonerTrackExists, UCEClonerComponent* /** InCloner */, uint32& /** OutTrackCount */)
+	static FOnClonerTrackExists OnClonerTrackExists;
+
+	FMovieSceneClonerTrackEditor(const TSharedRef<ISequencer>& InSequencer)
+		: FMovieSceneTrackEditor(InSequencer)
+	{}
+
+	virtual ~FMovieSceneClonerTrackEditor() override;
+
+	static TSharedRef<ISequencerTrackEditor> CreateTrackEditor(TSharedRef<ISequencer> InSequencer)
+	{
+		TSharedRef<FMovieSceneClonerTrackEditor> TrackEditor = MakeShared<FMovieSceneClonerTrackEditor>(InSequencer);
+		TrackEditor->BindDelegates();
+		return TrackEditor;
+	}
+
+private:
+	//~ Begin FMovieSceneTrackEditor
+	virtual bool SupportsType(TSubclassOf<UMovieSceneTrack> InTrackClass) const override;
+	//~ End FMovieSceneTrackEditor
+
+	void BindDelegates();
+	void ExecuteAddTrack(UCEClonerComponent* InCloner);
+	void ExecuteTrackExists(UCEClonerComponent* InCloner, uint32& OutCount) const;
+
+};
\ No newline at end of file
