MotionDesign : ClonerEffector
- Adding plugin 'NiagaraSimCaching' as a dependency

diff --git a/Engine/Plugins/Experimental/ClonerEffector/ClonerEffector.uplugin b/Engine/Plugins/Experimental/ClonerEffector/ClonerEffector.uplugin
--- a/Engine/Plugins/Experimental/ClonerEffector/ClonerEffector.uplugin
+++ b/Engine/Plugins/Experimental/ClonerEffector/ClonerEffector.uplugin
@@ -46,6 +46,10 @@
 		{
 			"Name": "Niagara",
 			"Enabled": true
+		},
+		{
+			"Name": "NiagaraSimCaching",
+			"Enabled": true
 		}
 	]
 }
\ No newline at end of file
