MotionDesign : ClonerEffector
- Fix long loading time / not appearing when spawning a cloner actor in editor

diff --git a/Engine/Plugins/Experimental/ClonerEffector/Source/ClonerEffector/Private/Cloner/CEClonerComponent.cpp b/Engine/Plugins/Experimental/ClonerEffector/Source/ClonerEffector/Private/Cloner/CEClonerComponent.cpp
--- a/Engine/Plugins/Experimental/ClonerEffector/Source/ClonerEffector/Private/Cloner/CEClonerComponent.cpp
+++ b/Engine/Plugins/Experimental/ClonerEffector/Source/ClonerEffector/Private/Cloner/CEClonerComponent.cpp
@@ -665,7 +665,7 @@ void UCEClonerComponent::UpdateDirtyMeshesAsync()
 
 	// Update baked dynamic meshes on other thread
 	TWeakObjectPtr<UCEClonerComponent> ThisWeak(this);
-	Async(EAsyncExecution::TaskGraph, [ThisWeak, DirtyAttachments]()
+	Async(EAsyncExecution::ThreadPool, [ThisWeak, DirtyAttachments]()
 	{
 		UCEClonerComponent* This = ThisWeak.Get();
 
@@ -1006,16 +1006,26 @@ void UCEClonerComponent::InitializeCloner()
 
 	SetAsset(nullptr);
 
-	OnLayoutNameChanged();
-
 #if WITH_EDITOR
 	OnVisualizerSpriteVisibleChanged();
-#endif
 
-	OnClonerInitializedDelegate.Broadcast(this);
+	const AActor* Owner = GetOwner();
+
+	// Skip init for preview actor
+	if (Owner && Owner->bIsEditorPreviewActor)
+	{
+		return;
+	}
+#endif
 
-	/** Register a custom ticker to avoid using the component tick that needs the simulation to be solo */
+	// Register a custom ticker to avoid using the component tick that needs the simulation to be solo
+	TreeUpdateDeltaTime = TreeUpdateInterval;
 	RegisterTicker();
+
+	// Load layout after registering ticker to let attachment tree update, then layout rendering will occur with up to date attachment data
+	OnLayoutNameChanged();
+
+	OnClonerInitializedDelegate.Broadcast(this);
 }
 
 void UCEClonerComponent::RegisterTicker()
@@ -1617,7 +1627,11 @@ UCEClonerLayoutBase* UCEClonerComponent::FindOrAddLayout(FName InLayoutName)
 	if (!NewActiveLayout)
 	{
 		NewActiveLayout = Subsystem->CreateNewLayout(InLayoutName, this);
-		LayoutInstances.Add(NewActiveLayout);
+
+		if (NewActiveLayout)
+		{
+			LayoutInstances.Add(NewActiveLayout);
+		}
 	}
 
 	return NewActiveLayout;
diff --git a/Engine/Plugins/Experimental/ClonerEffector/Source/ClonerEffector/Private/Cloner/Extensions/CEClonerConstraintExtension.cpp b/Engine/Plugins/Experimental/ClonerEffector/Source/ClonerEffector/Private/Cloner/Extensions/CEClonerConstraintExtension.cpp
--- a/Engine/Plugins/Experimental/ClonerEffector/Source/ClonerEffector/Private/Cloner/Extensions/CEClonerConstraintExtension.cpp
+++ b/Engine/Plugins/Experimental/ClonerEffector/Source/ClonerEffector/Private/Cloner/Extensions/CEClonerConstraintExtension.cpp
@@ -188,8 +188,10 @@ void UCEClonerConstraintExtension::OnExtensionParametersChanged(UCEClonerCompone
 
 	// Texture
 	static const FNiagaraVariable ConstraintTextureSamplerVar(FNiagaraTypeDefinition(UNiagaraDataInterfaceTexture::StaticClass()), TEXT("ConstraintTextureSampler"));
-	UNiagaraDataInterfaceTexture* TextureSamplerDI = Cast<UNiagaraDataInterfaceTexture>(ExposedParameters.GetDataInterface(ConstraintTextureSamplerVar));
-	TextureSamplerDI->SetTexture(TextureAsset.Get());
+	if (UNiagaraDataInterfaceTexture* TextureSamplerDI = Cast<UNiagaraDataInterfaceTexture>(ExposedParameters.GetDataInterface(ConstraintTextureSamplerVar)))
+	{
+		TextureSamplerDI->SetTexture(TextureAsset.Get());
+	}
 
 	static const FNiagaraVariable ConstraintTexturePlaneVar(FNiagaraTypeDefinition(StaticEnum<ECEClonerPlane>()), TEXT("ConstraintTexturePlane"));
 	ExposedParameters.SetParameterValue<int32>(static_cast<int32>(TexturePlane), ConstraintTexturePlaneVar);
diff --git a/Engine/Plugins/Experimental/ClonerEffector/Source/ClonerEffector/Private/Cloner/Layouts/CEClonerLayoutBase.cpp b/Engine/Plugins/Experimental/ClonerEffector/Source/ClonerEffector/Private/Cloner/Layouts/CEClonerLayoutBase.cpp
--- a/Engine/Plugins/Experimental/ClonerEffector/Source/ClonerEffector/Private/Cloner/Layouts/CEClonerLayoutBase.cpp
+++ b/Engine/Plugins/Experimental/ClonerEffector/Source/ClonerEffector/Private/Cloner/Layouts/CEClonerLayoutBase.cpp
@@ -123,25 +123,15 @@ void UCEClonerLayoutBase::LoadLayout()
 
 	UE_LOG(LogCEClonerLayoutBase, Verbose, TEXT("%s : Cloner layout load requested %s - Template system %s - Package %s"), *GetClonerActor()->GetActorNameOrLabel(), *LayoutName.ToString(), *LayoutAssetPath, *CustomPackagePath.GetPackageFName().ToString())
 
-	TWeakObjectPtr<UCEClonerLayoutBase> ThisWeak(this);
-	Async(EAsyncExecution::ThreadPool, [ThisWeak, CustomPackagePath, LayoutPackagePath]()
-	{
-		UCEClonerLayoutBase* This = ThisWeak.Get();
-
-		if (!This)
-		{
-			return;
-		}
-
-		FLoadPackageAsyncOptionalParams Params;
-		Params.PackagePriority = INT32_MAX;
-		Params.CustomPackageName = CustomPackagePath.GetPackageFName();
-		Params.CompletionDelegate = MakeUnique<FLoadPackageAsyncDelegate>(FLoadPackageAsyncDelegate::CreateUObject(This, &UCEClonerLayoutBase::OnSystemPackageLoaded));
+	FLoadPackageAsyncOptionalParams Params;
+	Params.PackagePriority = INT32_MAX;
+	Params.LoadFlags = LOAD_Async | LOAD_MemoryReader | LOAD_DisableCompileOnLoad;
+	Params.CustomPackageName = CustomPackagePath.GetPackageFName();
+	Params.CompletionDelegate = MakeUnique<FLoadPackageAsyncDelegate>(FLoadPackageAsyncDelegate::CreateUObject(this, &UCEClonerLayoutBase::OnSystemPackageLoaded));
 
-		This->LoadRequestIdentifier = LoadPackageAsync(LayoutPackagePath, MoveTemp(Params));
+	LoadRequestIdentifier = LoadPackageAsync(LayoutPackagePath, MoveTemp(Params));
 
-		This->BindCleanupDelegates();
-	});
+	BindCleanupDelegates();
 }
 
 bool UCEClonerLayoutBase::UnloadLayout()
@@ -162,6 +152,8 @@ bool UCEClonerLayoutBase::UnloadLayout()
 #if WITH_EDITOR
 	NiagaraSystem->KillAllActiveCompilations();
 #endif
+
+	// Prevent GC Leak
 	UPackage* Package = NiagaraSystem->GetPackage();
 	Package->ClearFlags(RF_Standalone);
 	Package->MarkAsGarbage();
@@ -234,6 +226,10 @@ bool UCEClonerLayoutBase::DeactivateLayout()
 	ClonerComponent->GetOverrideParameters().Empty(/** ClearBindings */true);
 	ClonerComponent->SetAsset(nullptr);
 
+#if WITH_EDITOR
+	NiagaraSystem->KillAllActiveCompilations();
+#endif
+
 	UE_LOG(LogCEClonerLayoutBase, Verbose, TEXT("%s : Cloner layout deactivated %s"), *GetClonerActor()->GetActorNameOrLabel(), *LayoutName.ToString())
 
 	OnLayoutInactive();
@@ -312,60 +308,52 @@ void UCEClonerLayoutBase::OnSystemPackageLoaded(const FName& InName, UPackage* I
 		NiagaraSystem->RemoveFromRoot();
 		NiagaraSystem->ClearFlags(RF_Standalone);
 
-		TWeakObjectPtr<UCEClonerLayoutBase> ThisWeak(this);
-		Async(EAsyncExecution::TaskGraphMainThread, [ThisWeak, InName]()
-		{
-			UCEClonerLayoutBase* This = ThisWeak.Get();
+		CacheMeshRenderer();
+	}
 
-			if (!This)
-			{
-				return;
-			}
+	const bool bLayoutLoaded = IsLayoutLoaded();
 
-			This->NiagaraSystem->EnsureFullyLoaded();
+	if (bLayoutLoaded)
+	{
+		UE_LOG(LogCEClonerLayoutBase, Verbose, TEXT("%s : Cloner layout loaded %s - Template system %s - Package %s"), *GetClonerActor()->GetActorNameOrLabel(), *LayoutName.ToString(), *LayoutAssetPath, *InName.ToString())
 
-#if WITH_EDITOR
-			This->NiagaraSystem->WaitForCompilationComplete(/** GPUShaders */true, /** ShowProgress */false);
-#endif
+		OnLayoutLoaded();
+	}
+	else
+	{
+		UE_LOG(LogCEClonerLayoutBase, Warning, TEXT("%s : Cloner layout load failed %s - Template system %s - Package %s"), *GetClonerActor()->GetActorNameOrLabel(), *LayoutName.ToString(), *LayoutAssetPath, *InName.ToString())
+	}
+
+	OnClonerLayoutLoadedDelegate.Broadcast(this, bLayoutLoaded);
+	OnClonerLayoutLoadedDelegate.Clear();
+}
 
-			for (FNiagaraEmitterHandle& SystemEmitterHandle : This->NiagaraSystem->GetEmitterHandles())
+void UCEClonerLayoutBase::CacheMeshRenderer()
+{
+	if (!NiagaraSystem)
+	{
+		return;
+	}
+
+	for (FNiagaraEmitterHandle& SystemEmitterHandle : NiagaraSystem->GetEmitterHandles())
+	{
+		if (const FVersionedNiagaraEmitterData* EmitterData = SystemEmitterHandle.GetEmitterData())
+		{
+			for (UNiagaraRendererProperties* EmitterRenderer : EmitterData->GetRenderers())
 			{
-				if (const FVersionedNiagaraEmitterData* EmitterData = SystemEmitterHandle.GetEmitterData())
+				if (UNiagaraMeshRendererProperties* EmitterMeshRenderer = Cast<UNiagaraMeshRendererProperties>(EmitterRenderer))
 				{
-					for (UNiagaraRendererProperties* EmitterRenderer : EmitterData->GetRenderers())
-					{
-						if (UNiagaraMeshRendererProperties* EmitterMeshRenderer = Cast<UNiagaraMeshRendererProperties>(EmitterRenderer))
-						{
-							EmitterMeshRenderer->Meshes.Empty();
+					EmitterMeshRenderer->Meshes.Empty();
 #if WITH_EDITOR
-							EmitterMeshRenderer->OnMeshChanged();
+					EmitterMeshRenderer->OnMeshChanged();
 #endif
 
-							This->MeshRenderer = EmitterMeshRenderer;
+					MeshRenderer = EmitterMeshRenderer;
 
-							UE_LOG(LogCEClonerLayoutBase, Verbose, TEXT("%s : Cloner layout loaded %s - Template system %s - Package %s"), *This->GetClonerActor()->GetActorNameOrLabel(), *This->LayoutName.ToString(), *This->LayoutAssetPath, *InName.ToString())
-						}
-					}
+					return;
 				}
 			}
-
-			const bool bLayoutLoaded = This->IsLayoutLoaded();
-
-			if (bLayoutLoaded)
-			{
-				This->OnLayoutLoaded();
-			}
-
-			This->OnClonerLayoutLoadedDelegate.Broadcast(This, bLayoutLoaded);
-			This->OnClonerLayoutLoadedDelegate.Clear();
-		});
-	}
-	else
-	{
-		UE_LOG(LogCEClonerLayoutBase, Warning, TEXT("%s : Cloner layout load failed %s - Template system %s - Package %s"), *GetClonerActor()->GetActorNameOrLabel(), *LayoutName.ToString(), *LayoutAssetPath, *InName.ToString())
-
-		OnClonerLayoutLoadedDelegate.Broadcast(this, IsLayoutLoaded());
-		OnClonerLayoutLoadedDelegate.Clear();
+		}
 	}
 }
 
@@ -413,6 +401,7 @@ void UCEClonerLayoutBase::OnLevelCleanup()
 		UE_LOG(LogCEClonerLayoutBase, Log, TEXT("%s : Cloner layout cleanup %s"), *GetClonerActor()->GetActorNameOrLabel(), *LayoutName.ToString())
 
 		DeactivateLayout();
+
 		UnloadLayout();
 	}
 
diff --git a/Engine/Plugins/Experimental/ClonerEffector/Source/ClonerEffector/Public/Cloner/Layouts/CEClonerLayoutBase.h b/Engine/Plugins/Experimental/ClonerEffector/Source/ClonerEffector/Public/Cloner/Layouts/CEClonerLayoutBase.h
--- a/Engine/Plugins/Experimental/ClonerEffector/Source/ClonerEffector/Public/Cloner/Layouts/CEClonerLayoutBase.h
+++ b/Engine/Plugins/Experimental/ClonerEffector/Source/ClonerEffector/Public/Cloner/Layouts/CEClonerLayoutBase.h
@@ -139,8 +139,12 @@ protected:
 private:
 	FOnClonerLayoutLoaded OnClonerLayoutLoadedDelegate;
 
+	/** Called when the system package was async loaded */
 	void OnSystemPackageLoaded(const FName& InName, UPackage* InPackage, EAsyncLoadingResult::Type InResult);
 
+	/** Finds and cache first mesh renderer in emitter */
+	void CacheMeshRenderer();
+
 	/** Bind delegates to clear resources during level or world cleanup */
 	void BindCleanupDelegates();
 	void UnbindCleanupDelegates() const;
