MotionDesign : ClonerEffector
- Fixed default tool action not transacting
- Fixed default actor spawned when existing attachment are present in the cloner

diff --git a/Engine/Plugins/Experimental/Avalanche/Source/AvalancheInteractiveTools/Private/Tools/AvaInteractiveToolsActorToolBase.cpp b/Engine/Plugins/Experimental/Avalanche/Source/AvalancheInteractiveTools/Private/Tools/AvaInteractiveToolsActorToolBase.cpp
--- a/Engine/Plugins/Experimental/Avalanche/Source/AvalancheInteractiveTools/Private/Tools/AvaInteractiveToolsActorToolBase.cpp
+++ b/Engine/Plugins/Experimental/Avalanche/Source/AvalancheInteractiveTools/Private/Tools/AvaInteractiveToolsActorToolBase.cpp
@@ -15,7 +15,12 @@ bool UAvaInteractiveToolsActorToolBase::OnBegin()
 
 void UAvaInteractiveToolsActorToolBase::DefaultAction()
 {
-	SpawnedActor = SpawnActor(ActorClass, false);
+	if (OnBegin())
+	{
+		SpawnedActor = SpawnActor(ActorClass, /** Preview */false);
+
+		OnComplete();
+	}
 
 	Super::DefaultAction();
 }
diff --git a/Engine/Plugins/Experimental/Avalanche/Source/AvalancheInteractiveTools/Private/Tools/AvaInteractiveToolsToolBase.cpp b/Engine/Plugins/Experimental/Avalanche/Source/AvalancheInteractiveTools/Private/Tools/AvaInteractiveToolsToolBase.cpp
--- a/Engine/Plugins/Experimental/Avalanche/Source/AvalancheInteractiveTools/Private/Tools/AvaInteractiveToolsToolBase.cpp
+++ b/Engine/Plugins/Experimental/Avalanche/Source/AvalancheInteractiveTools/Private/Tools/AvaInteractiveToolsToolBase.cpp
@@ -219,9 +219,15 @@ AActor* UAvaInteractiveToolsToolBase::SpawnActor(TSubclassOf<AActor> InActorClas
 	{
 		if (FViewport* Viewport = ContextAPI->GetFocusedViewport())
 		{
-			if (TSharedPtr<IAvaViewportClient> AvaViewportClient = FAvaViewportUtils::GetAvaViewportClient(Viewport))
+			if (FEditorViewportClient* ViewportClient = FAvaViewportUtils::GetAsEditorViewportClient(Viewport))
 			{
-				const FVector2f ViewportSize = AvaViewportClient->GetViewportSize();
+				if (TSharedPtr<IAvaViewportClient> AvaViewportClient = FAvaViewportUtils::GetAsAvaViewportClient(ViewportClient))
+				{
+					const FVector2f ViewportSize = AvaViewportClient->GetViewportSize();
+					return SpawnActor(InActorClass, EAvaViewportStatus::Focused, ViewportSize * 0.5, bInPreview, InActorLabelOverride);
+				}
+
+				const FVector2f ViewportSize = Viewport->GetSizeXY();
 				return SpawnActor(InActorClass, EAvaViewportStatus::Focused, ViewportSize * 0.5, bInPreview, InActorLabelOverride);
 			}
 		}
diff --git a/Engine/Plugins/Experimental/ClonerEffector/Source/ClonerEffector/Private/Cloner/CEClonerActor.cpp b/Engine/Plugins/Experimental/ClonerEffector/Source/ClonerEffector/Private/Cloner/CEClonerActor.cpp
--- a/Engine/Plugins/Experimental/ClonerEffector/Source/ClonerEffector/Private/Cloner/CEClonerActor.cpp
+++ b/Engine/Plugins/Experimental/ClonerEffector/Source/ClonerEffector/Private/Cloner/CEClonerActor.cpp
@@ -69,6 +69,8 @@ ACEClonerActor::ACEClonerActor()
 		{
 			GEditor->GetSelectedActors()->SelectionChangedEvent.AddUObject(this, &ACEClonerActor::OnEditorSelectionChanged);
 		}
+
+		UCEClonerComponent::OnClonerMeshUpdated().AddUObject(this, &ACEClonerActor::SpawnDefaultActorAttached);
 #endif
 	}
 }
@@ -107,11 +109,6 @@ void ACEClonerActor::PostActorCreated()
 
 #if WITH_EDITOR
 	bSpawnDefaultActorAttached = true;
-
-	if (ClonerComponent->bClonerInitialized)
-	{
-		SpawnDefaultActorAttached();
-	}
 #endif
 }
 
@@ -363,12 +360,18 @@ void ACEClonerActor::MigrateDeprecatedProperties()
 }
 
 #if WITH_EDITOR
-void ACEClonerActor::SpawnDefaultActorAttached()
+void ACEClonerActor::SpawnDefaultActorAttached(UCEClonerComponent* InComponent)
 {
-	if (bSpawnDefaultActorAttached)
+	if (InComponent
+		&& InComponent == ClonerComponent
+		&& bSpawnDefaultActorAttached)
 	{
 		bSpawnDefaultActorAttached = false;
-		ClonerComponent->CreateDefaultActorAttached();
+
+		if (InComponent->GetAttachmentCount() == 0)
+		{
+			ClonerComponent->CreateDefaultActorAttached();
+		}
 	}
 }
 
diff --git a/Engine/Plugins/Experimental/ClonerEffector/Source/ClonerEffector/Public/Cloner/CEClonerActor.h b/Engine/Plugins/Experimental/ClonerEffector/Source/ClonerEffector/Public/Cloner/CEClonerActor.h
--- a/Engine/Plugins/Experimental/ClonerEffector/Source/ClonerEffector/Public/Cloner/CEClonerActor.h
+++ b/Engine/Plugins/Experimental/ClonerEffector/Source/ClonerEffector/Public/Cloner/CEClonerActor.h
@@ -53,7 +53,7 @@ private:
 	void MigrateDeprecatedProperties();
 
 #if WITH_EDITOR
-	void SpawnDefaultActorAttached();
+	void SpawnDefaultActorAttached(UCEClonerComponent* InComponent);
 
 	/** Hide outline selection when this cloner only is selected in viewport */
 	void OnEditorSelectionChanged(UObject* InSelection);
