MotionDesign : ClonerEffector
- Create niagara sequencer track moved to utilities category
- Added multi cloner support when adding tracks to sequencer

diff --git a/Engine/Plugins/Experimental/ClonerEffector/Source/ClonerEffector/Public/Cloner/Extensions/CEClonerCollisionExtension.h b/Engine/Plugins/Experimental/ClonerEffector/Source/ClonerEffector/Public/Cloner/Extensions/CEClonerCollisionExtension.h
--- a/Engine/Plugins/Experimental/ClonerEffector/Source/ClonerEffector/Public/Cloner/Extensions/CEClonerCollisionExtension.h
+++ b/Engine/Plugins/Experimental/ClonerEffector/Source/ClonerEffector/Public/Cloner/Extensions/CEClonerCollisionExtension.h
@@ -148,11 +148,11 @@ protected:
 	FVector CollisionGridSize = FVector(5000.f);
 
 	/** Collision radius calculation mode */
-	UPROPERTY(EditInstanceOnly, Setter, Getter, Category="Collisions", meta=(EditCondition="bParticleCollisionEnabled", EditConditionHides))
+	UPROPERTY(EditInstanceOnly, Setter, Getter, Category="Collisions", meta=(EditCondition="bSurfaceCollisionEnabled || bParticleCollisionEnabled", EditConditionHides))
 	ECEClonerCollisionRadiusMode CollisionRadiusMode = ECEClonerCollisionRadiusMode::ExtentLength;
 
 	/** Radius expected around each particle for collision, order matches attachment index */
-	UPROPERTY(EditInstanceOnly, Getter, Category="Collisions", EditFixedSize, meta=(ClampMin="0", EditFixedOrder, EditCondition="bParticleCollisionEnabled && CollisionRadiusMode == ECEClonerCollisionRadiusMode::Manual", EditConditionHides))
+	UPROPERTY(EditInstanceOnly, Getter, Category="Collisions", EditFixedSize, meta=(ClampMin="0", EditFixedOrder, EditCondition="(bSurfaceCollisionEnabled || bParticleCollisionEnabled) && CollisionRadiusMode == ECEClonerCollisionRadiusMode::Manual", EditConditionHides))
 	TArray<float> CollisionRadii;
 
 	/** Minimum particle mass, used for collisions to push apart */
diff --git a/Engine/Plugins/Experimental/ClonerEffector/Source/ClonerEffectorEditor/Private/Cloner/Customizations/CEEditorClonerComponentDetailCustomization.cpp b/Engine/Plugins/Experimental/ClonerEffector/Source/ClonerEffectorEditor/Private/Cloner/Customizations/CEEditorClonerComponentDetailCustomization.cpp
--- a/Engine/Plugins/Experimental/ClonerEffector/Source/ClonerEffectorEditor/Private/Cloner/Customizations/CEEditorClonerComponentDetailCustomization.cpp
+++ b/Engine/Plugins/Experimental/ClonerEffector/Source/ClonerEffectorEditor/Private/Cloner/Customizations/CEEditorClonerComponentDetailCustomization.cpp
@@ -5,6 +5,7 @@
 #include "Cloner/CEClonerComponent.h"
 #include "Cloner/Extensions/CEClonerExtensionBase.h"
 #include "Cloner/Layouts/CEClonerLayoutBase.h"
+#include "Cloner/Sequencer/MovieSceneClonerTrackEditor.h"
 #include "DetailLayoutBuilder.h"
 #include "DetailWidgetRow.h"
 #include "IDetailChildrenBuilder.h"
@@ -140,7 +141,10 @@ void FCEEditorClonerComponentDetailCustomization::CustomizeDetails(IDetailLayout
 		}
 	}
 
-	if (!LayoutFunctionNames.IsEmpty())
+	// Add cloner create track button
+	FunctionToCategory.Add(TrackEditor, TEXT("Utilities"));
+
+	if (!FunctionToCategory.IsEmpty())
 	{
 		TMap<FName, const TSharedPtr<SVerticalBox>> FunctionToWidget;
 
@@ -170,11 +174,11 @@ void FCEEditorClonerComponentDetailCustomization::CustomizeDetails(IDetailLayout
 			FunctionToWidget.Add(FunctionToCategoryPair.Key, FunctionsWidget);
 		}
 
-		for (const TPair<FName, TMap<TWeakObjectPtr<UObject>, TWeakObjectPtr<UFunction>>>& LayoutFunctionNamesPair : LayoutFunctionNames)
+		for (const TPair<FName, const TSharedPtr<SVerticalBox>>& FunctionNameToWidget : FunctionToWidget)
 		{
-			const FText ButtonLabel = FText::FromString(FName::NameToDisplayString(LayoutFunctionNamesPair.Key.ToString(), false));
+			const FText ButtonLabel = FText::FromString(FName::NameToDisplayString(FunctionNameToWidget.Key.ToString(), /** IsBool */false));
 
-			const TSharedPtr<SVerticalBox>& FunctionsWidget = FunctionToWidget.FindChecked(LayoutFunctionNamesPair.Key);
+			const TSharedPtr<SVerticalBox>& FunctionsWidget = FunctionNameToWidget.Value;
 
 			FunctionsWidget->AddSlot()
 				.Padding(0.f, 3.f)
@@ -184,7 +188,8 @@ void FCEEditorClonerComponentDetailCustomization::CustomizeDetails(IDetailLayout
 					.Text(ButtonLabel)
 					.HAlign(HAlign_Center)
 					.VAlign(VAlign_Fill)
-					.OnClicked(this, &FCEEditorClonerComponentDetailCustomization::OnFunctionButtonClicked, LayoutFunctionNamesPair.Key)
+					.IsEnabled(this, &FCEEditorClonerComponentDetailCustomization::IsFunctionButtonEnabled, FunctionNameToWidget.Key)
+					.OnClicked(this, &FCEEditorClonerComponentDetailCustomization::OnFunctionButtonClicked, FunctionNameToWidget.Key)
 				];
 		}
 	}
@@ -240,9 +245,23 @@ void FCEEditorClonerComponentDetailCustomization::OnClonerLayoutLoaded(UCECloner
 	}
 }
 
+bool FCEEditorClonerComponentDetailCustomization::IsFunctionButtonEnabled(FName InFunctionName) const
+{
+	if (InFunctionName.IsEqual(TrackEditor))
+	{
+		return CanAddSequencerTracks();
+	}
+
+	return true;
+}
+
 FReply FCEEditorClonerComponentDetailCustomization::OnFunctionButtonClicked(FName InFunctionName)
 {
-	if (TMap<TWeakObjectPtr<UObject>, TWeakObjectPtr<UFunction>> const* ObjectFunctions = LayoutFunctionNames.Find(InFunctionName))
+	if (InFunctionName.IsEqual(TrackEditor))
+	{
+		OnAddSequencerTracks();
+	}
+	else if (TMap<TWeakObjectPtr<UObject>, TWeakObjectPtr<UFunction>> const* ObjectFunctions = LayoutFunctionNames.Find(InFunctionName))
 	{
 		for (const TPair<TWeakObjectPtr<UObject>, TWeakObjectPtr<UFunction>>& ObjectFunction : *ObjectFunctions)
 		{
@@ -261,4 +280,48 @@ FReply FCEEditorClonerComponentDetailCustomization::OnFunctionButtonClicked(FNam
 	return FReply::Handled();
 }
 
+bool FCEEditorClonerComponentDetailCustomization::CanAddSequencerTracks() const
+{
+	// Lifecycle + cache tracks
+	constexpr uint32 ExpectedTrackCount = 2;
+
+	for (const TWeakObjectPtr<UCEClonerComponent>& ClonerComponentWeak : ClonerComponentsWeak)
+	{
+		if (UCEClonerComponent* ClonerComponent = ClonerComponentWeak.Get())
+		{
+			uint32 TrackCount = 0;
+			FMovieSceneClonerTrackEditor::OnClonerTrackExists.Broadcast(ClonerComponent, TrackCount);
+
+			if (TrackCount < ExpectedTrackCount)
+			{
+				return true;
+			}
+		}
+	}
+
+	return false;
+}
+
+FReply FCEEditorClonerComponentDetailCustomization::OnAddSequencerTracks()
+{
+	TSet<UCEClonerComponent*> ClonersWeak;
+
+	Algo::TransformIf(
+		ClonerComponentsWeak
+		, ClonersWeak
+		, [](const TWeakObjectPtr<UCEClonerComponent>& InClonerWeak)
+		{
+			return InClonerWeak.IsValid();
+		}
+		, [](const TWeakObjectPtr<UCEClonerComponent>& InClonerWeak)
+		{
+			return InClonerWeak.Get();
+		}
+	);
+
+	FMovieSceneClonerTrackEditor::OnAddClonerTrack.Broadcast(ClonersWeak);
+
+	return FReply::Handled();
+}
+
 #undef LOCTEXT_NAMESPACE
diff --git a/Engine/Plugins/Experimental/ClonerEffector/Source/ClonerEffectorEditor/Private/Cloner/Customizations/CEEditorClonerComponentDetailCustomization.h b/Engine/Plugins/Experimental/ClonerEffector/Source/ClonerEffectorEditor/Private/Cloner/Customizations/CEEditorClonerComponentDetailCustomization.h
--- a/Engine/Plugins/Experimental/ClonerEffector/Source/ClonerEffectorEditor/Private/Cloner/Customizations/CEEditorClonerComponentDetailCustomization.h
+++ b/Engine/Plugins/Experimental/ClonerEffector/Source/ClonerEffectorEditor/Private/Cloner/Customizations/CEEditorClonerComponentDetailCustomization.h
@@ -37,16 +37,22 @@ public:
 	//~ End IDetailCustomization
 
 protected:
+	static constexpr const TCHAR* TrackEditor = TEXT("CreateClonerSequencerTracks");
 	static void RemoveEmptySections();
 	static void OnChildPropertyChanged(const FPropertyChangedEvent& InEvent, TWeakPtr<IPropertyHandle> InParentHandleWeak);
 	static void OnPropertyChanged(const FPropertyChangedEvent& InEvent, TWeakPtr<IPropertyUtilities> InUtilitiesWeak);
 
+	bool CanAddSequencerTracks() const;
+	FReply OnAddSequencerTracks();
+
 	/** Bind delegates needed */
 	void Init();
 
 	/** Used to refresh details view when layout changes */
 	void OnClonerLayoutLoaded(UCEClonerComponent* InCloner, UCEClonerLayoutBase* InLayout);
 
+	bool IsFunctionButtonEnabled(FName InFunctionName) const;
+
 	/** Execute ufunction with that name on selected objects */
 	FReply OnFunctionButtonClicked(FName InFunctionName);
 
diff --git a/Engine/Plugins/Experimental/ClonerEffector/Source/ClonerEffectorEditor/Private/Cloner/Customizations/CEEditorClonerLifetimeExtensionDetailCustomization.cpp b/Engine/Plugins/Experimental/ClonerEffector/Source/ClonerEffectorEditor/Private/Cloner/Customizations/CEEditorClonerLifetimeExtensionDetailCustomization.cpp
--- a/Engine/Plugins/Experimental/ClonerEffector/Source/ClonerEffectorEditor/Private/Cloner/Customizations/CEEditorClonerLifetimeExtensionDetailCustomization.cpp
+++ b/Engine/Plugins/Experimental/ClonerEffector/Source/ClonerEffectorEditor/Private/Cloner/Customizations/CEEditorClonerLifetimeExtensionDetailCustomization.cpp
@@ -3,7 +3,6 @@
 #include "Cloner/Customizations/CEEditorClonerLifetimeExtensionDetailCustomization.h"
 
 #include "Cloner/Extensions/CEClonerLifetimeExtension.h"
-#include "Cloner/Sequencer/MovieSceneClonerTrackEditor.h"
 #include "DetailBuilderTypes.h"
 #include "DetailCategoryBuilder.h"
 #include "DetailLayoutBuilder.h"
@@ -59,31 +58,6 @@ void FCEEditorClonerLifetimeExtensionDetailCustomization::CustomizeDetails(IDeta
 			const TAttribute<EVisibility> VisibilityAttr = TAttribute<EVisibility>::CreateSP(this, &FCEEditorClonerLifetimeExtensionDetailCustomization::GetCurveVisibility, LifetimeExtensionWeak);
 			Row->Visibility(VisibilityAttr);
 		}
-
-		if (LifetimeExtensionsWeak.Num() == 1)
-		{
-			const FText ButtonLabel = LOCTEXT("AddSequencerTrack", "Add Sequencer Tracks");
-
-			IDetailCategoryBuilder& LifetimeCategoryBuilder = InDetailBuilder.EditCategory(CategoryName);
-
-			LifetimeCategoryBuilder.AddCustomRow(ButtonLabel)
-			.WholeRowContent()
-			.HAlign(HAlign_Left)
-			[
-				SNew(SVerticalBox)
-				+ SVerticalBox::Slot()
-				.Padding(0.f, 3.f)
-				.AutoHeight()
-				[
-					SNew(SButton)
-					.Text(ButtonLabel)
-					.HAlign(HAlign_Center)
-					.VAlign(VAlign_Fill)
-					.IsEnabled(this, &FCEEditorClonerLifetimeExtensionDetailCustomization::CanAddSequencerTracks, LifetimeExtensionWeak)
-					.OnClicked(this, &FCEEditorClonerLifetimeExtensionDetailCustomization::OnAddSequencerTracks, LifetimeExtensionWeak)
-				]
-			];
-		}
 	}
 }
 
@@ -102,28 +76,4 @@ EVisibility FCEEditorClonerLifetimeExtensionDetailCustomization::GetCurveVisibil
 		: EVisibility::Collapsed;
 }
 
-bool FCEEditorClonerLifetimeExtensionDetailCustomization::CanAddSequencerTracks(TWeakObjectPtr<UCEClonerLifetimeExtension> InClonerLifetimeExtension) const
-{
-	uint32 TrackCount = 0;
-
-	if (const UCEClonerLifetimeExtension* LifetimeExtension = InClonerLifetimeExtension.Get())
-	{
-		FMovieSceneClonerTrackEditor::OnClonerTrackExists.Broadcast(LifetimeExtension->GetClonerComponent(), TrackCount);
-	}
-
-	// Lifecycle + cache tracks
-	constexpr uint32 ExpectedTrackCount = 2;
-	return TrackCount < ExpectedTrackCount;
-}
-
-FReply FCEEditorClonerLifetimeExtensionDetailCustomization::OnAddSequencerTracks(TWeakObjectPtr<UCEClonerLifetimeExtension> InClonerLifetimeExtension)
-{
-	if (const UCEClonerLifetimeExtension* LifetimeExtension = InClonerLifetimeExtension.Get())
-	{
-		FMovieSceneClonerTrackEditor::OnAddClonerTrack.Broadcast(LifetimeExtension->GetClonerComponent());
-	}
-
-	return FReply::Handled();
-}
-
 #undef LOCTEXT_NAMESPACE
diff --git a/Engine/Plugins/Experimental/ClonerEffector/Source/ClonerEffectorEditor/Private/Cloner/Customizations/CEEditorClonerLifetimeExtensionDetailCustomization.h b/Engine/Plugins/Experimental/ClonerEffector/Source/ClonerEffectorEditor/Private/Cloner/Customizations/CEEditorClonerLifetimeExtensionDetailCustomization.h
--- a/Engine/Plugins/Experimental/ClonerEffector/Source/ClonerEffectorEditor/Private/Cloner/Customizations/CEEditorClonerLifetimeExtensionDetailCustomization.h
+++ b/Engine/Plugins/Experimental/ClonerEffector/Source/ClonerEffectorEditor/Private/Cloner/Customizations/CEEditorClonerLifetimeExtensionDetailCustomization.h
@@ -28,9 +28,6 @@ public:
 protected:
 	EVisibility GetCurveVisibility(TWeakObjectPtr<UCEClonerLifetimeExtension> InExtensionWeak) const;
 
-	bool CanAddSequencerTracks(TWeakObjectPtr<UCEClonerLifetimeExtension> InClonerLifetimeExtension) const;
-	FReply OnAddSequencerTracks(TWeakObjectPtr<UCEClonerLifetimeExtension> InClonerLifetimeExtension);
-
 	TSharedPtr<IPropertyHandle> LifetimeEnabledPropertyHandle;
 	TSharedPtr<IPropertyHandle> LifetimeScaleEnabledPropertyHandle;
 };
diff --git a/Engine/Plugins/Experimental/ClonerEffector/Source/ClonerEffectorEditor/Private/Cloner/Sequencer/MovieSceneClonerTrackEditor.cpp b/Engine/Plugins/Experimental/ClonerEffector/Source/ClonerEffectorEditor/Private/Cloner/Sequencer/MovieSceneClonerTrackEditor.cpp
--- a/Engine/Plugins/Experimental/ClonerEffector/Source/ClonerEffectorEditor/Private/Cloner/Sequencer/MovieSceneClonerTrackEditor.cpp
+++ b/Engine/Plugins/Experimental/ClonerEffector/Source/ClonerEffectorEditor/Private/Cloner/Sequencer/MovieSceneClonerTrackEditor.cpp
@@ -31,9 +31,9 @@ void FMovieSceneClonerTrackEditor::BindDelegates()
 	OnClonerTrackExists.AddSP(this, &FMovieSceneClonerTrackEditor::ExecuteTrackExists);
 }
 
-void FMovieSceneClonerTrackEditor::ExecuteAddTrack(UCEClonerComponent* InCloner)
+void FMovieSceneClonerTrackEditor::ExecuteAddTrack(const TSet<UCEClonerComponent*>& InCloners)
 {
-	if (!IsValid(InCloner))
+	if (InCloners.IsEmpty())
 	{
 		return;
 	}
@@ -54,48 +54,60 @@ void FMovieSceneClonerTrackEditor::ExecuteAddTrack(UCEClonerComponent* InCloner)
 
 	FocusedMovieScene->Modify();
 
-	const FGuid ComponentBinding = SequencerPtr->GetHandleToObject(InCloner);
-	if (FocusedMovieScene->FindSpawnable(ComponentBinding))
+	for (UCEClonerComponent* Cloner : InCloners)
 	{
-		// we only want to add tracks for possessables
-		return;
-	}
-
-	// Add lifecycle track
-	UMovieSceneTrack* LifeCycleTrack = FocusedMovieScene->FindTrack(UMovieSceneNiagaraSystemTrack::StaticClass(), ComponentBinding);
-
-	if (!LifeCycleTrack)
-	{
-		UMovieSceneNiagaraSystemTrack* NiagaraSystemTrack = FocusedMovieScene->AddTrack<UMovieSceneNiagaraSystemTrack>(ComponentBinding);
-		NiagaraSystemTrack->SetDisplayName(LOCTEXT("ClonerLifeCycleTrackName", "Cloner Life Cycle"));
-
-		UMovieSceneNiagaraSystemSpawnSection* NiagaraSpawnSection = Cast<UMovieSceneNiagaraSystemSpawnSection>(NiagaraSystemTrack->CreateNewSection());
-		NiagaraSpawnSection->SetAgeUpdateMode(ENiagaraAgeUpdateMode::DesiredAge);
-
-		NiagaraSpawnSection->SetRange(TRange<FFrameNumber>(
-			FocusedMovieScene->GetPlaybackRange().GetLowerBound(),
-			FocusedMovieScene->GetPlaybackRange().GetUpperBound()
-		));
-
-		NiagaraSystemTrack->AddSection(*NiagaraSpawnSection);
-	}
-
-	// Add cache track
-	UMovieSceneTrack* CacheTrack = FocusedMovieScene->FindTrack(UMovieSceneNiagaraCacheTrack::StaticClass(), ComponentBinding);
-
-	if (!CacheTrack)
-	{
-		UMovieSceneNiagaraCacheTrack* NiagaraCacheTrack = FocusedMovieScene->AddTrack<UMovieSceneNiagaraCacheTrack>(ComponentBinding);
-		NiagaraCacheTrack->SetDisplayName(LOCTEXT("ClonerSimCacheTrackName", "Cloner Sim Cache"));
-
-		UMovieSceneNiagaraCacheSection* NiagaraCacheSection = Cast<UMovieSceneNiagaraCacheSection>(NiagaraCacheTrack->CreateNewSection());
-
-		NiagaraCacheSection->SetRange(TRange<FFrameNumber>(
-			FocusedMovieScene->GetPlaybackRange().GetLowerBound(),
-			FocusedMovieScene->GetPlaybackRange().GetUpperBound()
-		));
-
-		NiagaraCacheTrack->AddSection(*NiagaraCacheSection);
+		if (!IsValid(Cloner))
+		{
+			continue;
+		}
+
+		const FGuid ComponentBinding = SequencerPtr->GetHandleToObject(Cloner);
+		if (FocusedMovieScene->FindSpawnable(ComponentBinding))
+		{
+			// we only want to add tracks for possessables
+			return;
+		}
+
+		// Add lifecycle track
+		const UMovieSceneTrack* LifeCycleTrack = FocusedMovieScene->FindTrack(UMovieSceneNiagaraSystemTrack::StaticClass(), ComponentBinding);
+
+		if (!LifeCycleTrack)
+		{
+			if (UMovieSceneNiagaraSystemTrack* NiagaraSystemTrack = FocusedMovieScene->AddTrack<UMovieSceneNiagaraSystemTrack>(ComponentBinding))
+			{
+				NiagaraSystemTrack->SetDisplayName(LOCTEXT("ClonerLifeCycleTrackName", "Cloner Life Cycle"));
+
+				UMovieSceneNiagaraSystemSpawnSection* NiagaraSpawnSection = CastChecked<UMovieSceneNiagaraSystemSpawnSection>(NiagaraSystemTrack->CreateNewSection());
+				NiagaraSpawnSection->SetAgeUpdateMode(ENiagaraAgeUpdateMode::DesiredAge);
+
+				NiagaraSpawnSection->SetRange(TRange<FFrameNumber>(
+					FocusedMovieScene->GetPlaybackRange().GetLowerBound(),
+					FocusedMovieScene->GetPlaybackRange().GetUpperBound()
+				));
+
+				NiagaraSystemTrack->AddSection(*NiagaraSpawnSection);
+			}
+		}
+
+		// Add cache track
+		const UMovieSceneTrack* CacheTrack = FocusedMovieScene->FindTrack(UMovieSceneNiagaraCacheTrack::StaticClass(), ComponentBinding);
+
+		if (!CacheTrack)
+		{
+			if (UMovieSceneNiagaraCacheTrack* NiagaraCacheTrack = FocusedMovieScene->AddTrack<UMovieSceneNiagaraCacheTrack>(ComponentBinding))
+			{
+				NiagaraCacheTrack->SetDisplayName(LOCTEXT("ClonerSimCacheTrackName", "Cloner Sim Cache"));
+
+				UMovieSceneNiagaraCacheSection* NiagaraCacheSection = CastChecked<UMovieSceneNiagaraCacheSection>(NiagaraCacheTrack->CreateNewSection());
+
+				NiagaraCacheSection->SetRange(TRange<FFrameNumber>(
+					FocusedMovieScene->GetPlaybackRange().GetLowerBound(),
+					FocusedMovieScene->GetPlaybackRange().GetUpperBound()
+				));
+
+				NiagaraCacheTrack->AddSection(*NiagaraCacheSection);
+			}
+		}
 	}
 
 	SequencerPtr->NotifyMovieSceneDataChanged(EMovieSceneDataChangeType::MovieSceneStructureItemAdded);
diff --git a/Engine/Plugins/Experimental/ClonerEffector/Source/ClonerEffectorEditor/Private/Cloner/Sequencer/MovieSceneClonerTrackEditor.h b/Engine/Plugins/Experimental/ClonerEffector/Source/ClonerEffectorEditor/Private/Cloner/Sequencer/MovieSceneClonerTrackEditor.h
--- a/Engine/Plugins/Experimental/ClonerEffector/Source/ClonerEffectorEditor/Private/Cloner/Sequencer/MovieSceneClonerTrackEditor.h
+++ b/Engine/Plugins/Experimental/ClonerEffector/Source/ClonerEffectorEditor/Private/Cloner/Sequencer/MovieSceneClonerTrackEditor.h
@@ -11,7 +11,7 @@ class ISequencer;
 class FMovieSceneClonerTrackEditor : public FMovieSceneTrackEditor
 {
 public:
-	DECLARE_MULTICAST_DELEGATE_OneParam(FOnAddClonerTrack, UCEClonerComponent* /** InCloner */)
+	DECLARE_MULTICAST_DELEGATE_OneParam(FOnAddClonerTrack, const TSet<UCEClonerComponent*>& /** InCloners */)
 	static FOnAddClonerTrack OnAddClonerTrack;
 
 	DECLARE_MULTICAST_DELEGATE_TwoParams(FOnClonerTrackExists, UCEClonerComponent* /** InCloner */, uint32& /** OutTrackCount */)
@@ -36,7 +36,7 @@ private:
 	//~ End FMovieSceneTrackEditor
 
 	void BindDelegates();
-	void ExecuteAddTrack(UCEClonerComponent* InCloner);
+	void ExecuteAddTrack(const TSet<UCEClonerComponent*>& InCloners);
 	void ExecuteTrackExists(UCEClonerComponent* InCloner, uint32& OutCount) const;
 
 };
\ No newline at end of file
