MotionDesign : ClonerEffector
- Improving material change detection to solve issue with shape and material designer not yet compiled and missing flag

diff --git a/Engine/Plugins/Experimental/ClonerEffector/Source/ClonerEffector/Private/Cloner/CEClonerComponent.cpp b/Engine/Plugins/Experimental/ClonerEffector/Source/ClonerEffector/Private/Cloner/CEClonerComponent.cpp
--- a/Engine/Plugins/Experimental/ClonerEffector/Source/ClonerEffector/Private/Cloner/CEClonerComponent.cpp
+++ b/Engine/Plugins/Experimental/ClonerEffector/Source/ClonerEffector/Private/Cloner/CEClonerComponent.cpp
@@ -70,6 +70,8 @@ UCEClonerComponent::UCEClonerComponent()
 #if WITH_EDITOR
 		FCoreUObjectDelegates::OnObjectPropertyChanged.RemoveAll(this);
 		FCoreUObjectDelegates::OnObjectPropertyChanged.AddUObject(this, &UCEClonerComponent::OnActorPropertyChanged);
+
+		UMaterial::OnMaterialCompilationFinished().AddUObject(this, &UCEClonerComponent::OnMaterialCompiled);
 #endif
 
 		// Apply default layout
@@ -150,6 +152,15 @@ void UCEClonerComponent::OnComponentCreated()
 	InitializeCloner();
 }
 
+void UCEClonerComponent::OnComponentDestroyed(bool bInDestroyingHierarchy)
+{
+	Super::OnComponentDestroyed(bInDestroyingHierarchy);
+
+#if WITH_EDITOR
+	UMaterial::OnMaterialCompilationFinished().RemoveAll(this);
+#endif
+}
+
 void UCEClonerComponent::UpdateClonerRenderState()
 {
 	/**
@@ -481,6 +492,11 @@ void UCEClonerComponent::OnActorPropertyChanged(UObject* InObject, FPropertyChan
 {
 	OnMaterialChanged(InObject);
 }
+
+void UCEClonerComponent::OnMaterialCompiled(UMaterialInterface* InMaterial)
+{
+	OnMaterialChanged(InMaterial);
+}
 #endif
 
 void UCEClonerComponent::OnMaterialChanged(UObject* InObject)
@@ -519,7 +535,9 @@ void UCEClonerComponent::OnMaterialChanged(UObject* InObject)
 	bool bMaterialChanged = false;
 	TArray<TWeakObjectPtr<UMaterialInterface>> NewMaterials;
 	NewMaterials.Reserve(PrimitiveComponents.Num());
+	UMaterialInterface* DefaultMaterial = LoadObject<UMaterialInterface>(nullptr, UCEClonerEffectorSettings::DefaultMaterialPath);
 
+	TArray<TWeakObjectPtr<UMaterialInterface>> UnsetMaterialsWeak;
 	for (UPrimitiveComponent* PrimitiveComponent : PrimitiveComponents)
 	{
 		if (!PrimitiveComponent || !FCEMeshBuilder::HasAnyGeometry(PrimitiveComponent))
@@ -529,18 +547,31 @@ void UCEClonerComponent::OnMaterialChanged(UObject* InObject)
 
 		for (int32 MatIndex = 0; MatIndex < PrimitiveComponent->GetNumMaterials(); MatIndex++)
 		{
-			UMaterialInterface* Material = PrimitiveComponent->GetMaterial(MatIndex);
+			UMaterialInterface* PreviousMaterial = PrimitiveComponent->GetMaterial(MatIndex);
+			UMaterialInterface* NewMaterial = PreviousMaterial;
 
-			if (!AttachmentItem->BakedMaterials.IsValidIndex(MatIdx) || AttachmentItem->BakedMaterials[MatIdx] != Material)
+			if (FilterSupportedMaterial(NewMaterial, DefaultMaterial))
+			{
+				UnsetMaterialsWeak.Add(PreviousMaterial);
+			}
+
+			if (!AttachmentItem->BakedMaterials.IsValidIndex(MatIdx)
+				|| AttachmentItem->BakedMaterials[MatIdx] != NewMaterial)
 			{
 				bMaterialChanged = true;
 			}
 
-			NewMaterials.Add(Material);
+			NewMaterials.Add(NewMaterial);
 			MatIdx++;
 		}
 	}
 
+	// Show warning for unset materials
+	if (!UnsetMaterialsWeak.IsEmpty())
+	{
+		FireMaterialWarning(ActorChanged, UnsetMaterialsWeak);
+	}
+
 	if (bMaterialChanged)
 	{
 		UE_LOG(LogCEClonerComponent, Log, TEXT("%s : Detected material change for %s"), *ClonerActor->GetActorNameOrLabel(), *ActorChanged->GetActorNameOrLabel());
@@ -779,6 +810,14 @@ void UCEClonerComponent::UpdateActorBakedDynamicMesh(AActor* InActor)
 	MeshBuilder.Reset();
 
 	AttachmentItem->BakedMesh = Mesh;
+
+	TArray<TWeakObjectPtr<UMaterialInterface>> UnsetMaterials;
+	UMaterialInterface* DefaultMaterial = LoadObject<UMaterialInterface>(nullptr, UCEClonerEffectorSettings::DefaultMaterialPath);
+	if (FilterSupportedMaterials(MeshMaterials, UnsetMaterials, DefaultMaterial))
+	{
+		FireMaterialWarning(InActor, UnsetMaterials);
+	}
+
 	AttachmentItem->BakedMaterials = MoveTemp(MeshMaterials);
 
 	// Was the mesh invalidated during the update process ?
@@ -820,9 +859,6 @@ void UCEClonerComponent::UpdateRootActorBakedStaticMesh(AActor* InRootActor)
 	TArray<FCEClonerAttachmentItem*> AttachmentItems;
 	GetActorAttachmentItems(InRootActor, AttachmentItems);
 
-	int32 ReadOnlyMaterialCount = 0;
-	UMaterialInterface* DefaultClonerMaterial = nullptr;
-
 	for (FCEClonerAttachmentItem* AttachmentItem : AttachmentItems)
 	{
 		if (!AttachmentItem)
@@ -843,35 +879,7 @@ void UCEClonerComponent::UpdateRootActorBakedStaticMesh(AActor* InRootActor)
 			MeshTransform = AttachmentItem->ItemActor->GetTransform().GetRelativeTransform(ParentTransform);
 		}
 
-		TArray<TWeakObjectPtr<UMaterialInterface>> MeshMaterials;
-		for (TWeakObjectPtr<UMaterialInterface>& BakedMaterial : AttachmentItem->BakedMaterials)
-		{
-			UMaterialInterface* MaterialInterface = BakedMaterial.Get();
-
-			if (MaterialInterface && !IsMaterialUsageFlagSet(MaterialInterface))
-			{
-				if (!IsMaterialDirtyable(MaterialInterface))
-				{
-					ReadOnlyMaterialCount++;
-					UE_LOG(LogCEClonerComponent, Warning, TEXT("%s : The following materials (%s) on actor (%s) does not have the required usage flag (bUsedWithNiagaraMeshParticles) to work with the cloner, this material cannot be dirtied due to its read-only location, replacing material by default cloner material and proceeding"), *ClonerActor->GetActorNameOrLabel(), MaterialInterface ? *MaterialInterface->GetMaterial()->GetPathName() : TEXT("Invalid Material"), *InRootActor->GetActorNameOrLabel());
-
-					if (!DefaultClonerMaterial)
-					{
-						DefaultClonerMaterial = LoadObject<UMaterialInterface>(nullptr, UCEClonerEffectorSettings::DefaultMaterialPath);
-					}
-
-					MaterialInterface = DefaultClonerMaterial;
-				}
-				else
-				{
-					UE_LOG(LogCEClonerComponent, Log, TEXT("%s : The following materials (%s) on actor (%s) does not have the required usage flag (bUsedWithNiagaraMeshParticles) to work with the cloner, setting the flag on material and proceeding, please resave the asset with the flag set"), *ClonerActor->GetActorNameOrLabel(), MaterialInterface ? *MaterialInterface->GetMaterial()->GetPathName() : TEXT("Invalid Material"), *InRootActor->GetActorNameOrLabel());
-				}
-			}
-
-			MeshMaterials.Add(MaterialInterface);
-		}
-
-		MeshBuilder.AppendMesh(AttachmentItem->BakedMesh, MeshMaterials, MeshTransform);
+		MeshBuilder.AppendMesh(AttachmentItem->BakedMesh, AttachmentItem->BakedMaterials, MeshTransform);
 	}
 
 	UStaticMesh* Mesh = NewObject<UStaticMesh>();
@@ -880,23 +888,47 @@ void UCEClonerComponent::UpdateRootActorBakedStaticMesh(AActor* InRootActor)
 	MeshBuilder.Reset();
 
 	ClonerTree.MergedBakedMeshes[RootIdx] = Mesh;
+}
 
-	if (ReadOnlyMaterialCount > 0)
+bool UCEClonerComponent::FilterSupportedMaterials(TArray<TWeakObjectPtr<UMaterialInterface>>& InMaterials, TArray<TWeakObjectPtr<UMaterialInterface>>& OutUnsetMaterials, UMaterialInterface* InDefaultMaterial)
+{
+	check(InDefaultMaterial)
+
+	OutUnsetMaterials.Reset(InMaterials.Num());
+
+	for (int32 Index = 0; Index < InMaterials.Num(); Index++)
 	{
-		UE_LOG(LogCEClonerComponent, Warning, TEXT("%s : %i invalid material(s) detected due to missing niagara usage flag (bUsedWithNiagaraMeshParticles) on actor (%s)"), *ClonerActor->GetActorNameOrLabel(), ReadOnlyMaterialCount, *InRootActor->GetActorNameOrLabel());
+		UMaterialInterface* PreviousMaterialInterface = InMaterials[Index].Get();
 
-#if WITH_EDITOR
-		// Fire warning notification when invalid materials are found and at least 5s has elapsed since last notification
-		constexpr double MinNotificationElapsedTime = 5.0;
-		const double CurrentTime = FApp::GetCurrentTime();
+		UMaterialInterface* NewMaterialInterface = PreviousMaterialInterface;
 
-		if (CurrentTime - LastNotificationTime > MinNotificationElapsedTime)
+		if (FilterSupportedMaterial(NewMaterialInterface, InDefaultMaterial))
 		{
-			LastNotificationTime = CurrentTime;
-			ShowMaterialWarning(ReadOnlyMaterialCount);
+			// Add original material to unset list
+			OutUnsetMaterials.Add(PreviousMaterialInterface);
 		}
-#endif
+
+		// Replace material
+		InMaterials[Index] = NewMaterialInterface;
+	}
+
+	return OutUnsetMaterials.IsEmpty();
+}
+
+bool UCEClonerComponent::FilterSupportedMaterial(UMaterialInterface*& InMaterial, UMaterialInterface* InDefaultMaterial)
+{
+	if (InMaterial && !IsMaterialUsageFlagSet(InMaterial))
+	{
+		// Replace material if dirtyable and not in read only location
+		if (!IsMaterialDirtyable(InMaterial))
+		{
+			InMaterial = InDefaultMaterial;
+		}
+
+		return true;
 	}
+
+	return false;
 }
 
 void UCEClonerComponent::GetActorAttachmentItems(AActor* InActor, TArray<FCEClonerAttachmentItem*>& OutAttachmentItems)
@@ -938,7 +970,7 @@ bool UCEClonerComponent::IsAllMergedMeshesValid() const
 
 bool UCEClonerComponent::IsMaterialDirtyable(const UMaterialInterface* InMaterial)
 {
-	const UMaterial* BaseMaterial = InMaterial->GetMaterial();
+	const UMaterial* BaseMaterial = InMaterial->GetMaterial_Concurrent();
 	const FString ContentFolder = FPaths::ConvertRelativePathToFull(FPaths::ProjectContentDir());
 
 	const UPackage* MaterialPackage = BaseMaterial->GetPackage();
@@ -956,7 +988,7 @@ bool UCEClonerComponent::IsMaterialUsageFlagSet(const UMaterialInterface* InMate
 {
 	if (InMaterial)
 	{
-		if (const UMaterial* Material = InMaterial->GetMaterial())
+		if (const UMaterial* Material = InMaterial->GetMaterial_Concurrent())
 		{
 			return Material->GetUsageByFlag(EMaterialUsage::MATUSAGE_NiagaraMeshParticles);
 		}
@@ -965,12 +997,43 @@ bool UCEClonerComponent::IsMaterialUsageFlagSet(const UMaterialInterface* InMate
 	return false;
 }
 
+void UCEClonerComponent::FireMaterialWarning(const AActor* InContextActor, const TArray<TWeakObjectPtr<UMaterialInterface>>& InUnsetMaterials)
+{
+	if (!IsValid(InContextActor) || InUnsetMaterials.IsEmpty())
+	{
+		return;
+	}
+
+	UE_LOG(LogCEClonerComponent, Warning, TEXT("%s : %i unsupported material(s) detected due to missing niagara usage flag (bUsedWithNiagaraMeshParticles) on actor (%s), see logs below"), *GetOwner()->GetActorNameOrLabel(), InUnsetMaterials.Num(), *InContextActor->GetActorNameOrLabel());
+
+	const AActor* ClonerActor = GetOwner();
+	for (const TWeakObjectPtr<UMaterialInterface>& UnsetMaterialWeak : InUnsetMaterials)
+	{
+		if (UMaterialInterface* UnsetMaterial = UnsetMaterialWeak.Get())
+		{
+			UE_LOG(LogCEClonerComponent, Warning, TEXT("%s : The following materials (%s) on actor (%s) does not have the usage flag (bUsedWithNiagaraMeshParticles) set to work with the cloner, set the flag and resave the asset to avoid this warning"), *ClonerActor->GetActorNameOrLabel(), *UnsetMaterial->GetMaterial()->GetPathName(), *InContextActor->GetActorNameOrLabel());
+		}
+	}
+
+#if WITH_EDITOR
+	// Fire warning notification when invalid materials are found and at least 5s has elapsed since last notification
+	constexpr double MinNotificationElapsedTime = 5.0;
+	const double CurrentTime = FApp::GetCurrentTime();
+
+	if (CurrentTime - LastNotificationTime > MinNotificationElapsedTime)
+	{
+		LastNotificationTime = CurrentTime;
+		ShowMaterialWarning(InUnsetMaterials.Num());
+	}
+#endif
+}
+
 #if WITH_EDITOR
 void UCEClonerComponent::ShowMaterialWarning(int32 InMaterialCount)
 {
 	if (InMaterialCount > 0)
 	{
-		FNotificationInfo NotificationInfo(FText::Format(LOCTEXT("MaterialsMissingUsageFlag", "Detected {0} read-only material(s) with missing niagara usage flag required to work properly with cloner (See logs)"), InMaterialCount));
+		FNotificationInfo NotificationInfo(FText::Format(LOCTEXT("MaterialsMissingUsageFlag", "Detected {0} material(s) with missing usage flag required to work properly with cloner (See logs)"), InMaterialCount));
 		NotificationInfo.ExpireDuration = 5.f;
 		NotificationInfo.bFireAndForget = true;
 		NotificationInfo.Image = FAppStyle::GetBrush("Icons.WarningWithColor");
diff --git a/Engine/Plugins/Experimental/ClonerEffector/Source/ClonerEffector/Public/Cloner/CEClonerComponent.h b/Engine/Plugins/Experimental/ClonerEffector/Source/ClonerEffector/Public/Cloner/CEClonerComponent.h
--- a/Engine/Plugins/Experimental/ClonerEffector/Source/ClonerEffector/Public/Cloner/CEClonerComponent.h
+++ b/Engine/Plugins/Experimental/ClonerEffector/Source/ClonerEffector/Public/Cloner/CEClonerComponent.h
@@ -235,6 +235,14 @@ protected:
 	/** Called when cloner is initialized */
 	CLONEREFFECTOR_API static FOnClonerInitialized OnClonerInitializedDelegate;
 
+	/** Replaces all unsupported material by default material, gathers unset materials that needs recompiling with proper flags */
+	static bool FilterSupportedMaterials(TArray<TWeakObjectPtr<UMaterialInterface>>& InMaterials, TArray<TWeakObjectPtr<UMaterialInterface>>& OutUnsetMaterials, UMaterialInterface* InDefaultMaterial);
+
+	static bool FilterSupportedMaterial(UMaterialInterface*& InMaterial, UMaterialInterface* InDefaultMaterial);
+
+	/** Fires a warning about unset materials used within this cloner */
+	void FireMaterialWarning(const AActor* InContextActor, const TArray<TWeakObjectPtr<UMaterialInterface>>& InUnsetMaterials);
+
 	//~ Begin UObject
 	virtual void PostInitProperties() override;
 	virtual void PostLoad() override;
@@ -246,6 +254,7 @@ protected:
 
 	//~ Begin UActorComponent
 	virtual void OnComponentCreated() override;
+	virtual void OnComponentDestroyed(bool bInDestroyingHierarchy) override;
 	//~ End UActorComponent
 
 	void UpdateAttachmentTree();
@@ -367,6 +376,7 @@ private:
 
 #if WITH_EDITOR
 	void OnActorPropertyChanged(UObject* InObject, FPropertyChangedEvent& InPropertyChangedEvent);
+	void OnMaterialCompiled(UMaterialInterface* InMaterial);
 #endif
 
 	/** Called when a cloned actor is destroyed */
