MotionDesign : ClonerEffector
- Improve loading time by caching niagara system in the layout and saving it instead of recreating it every time

diff --git a/Engine/Plugins/Experimental/ClonerEffector/Source/ClonerEffector/Private/Cloner/CEClonerComponent.cpp b/Engine/Plugins/Experimental/ClonerEffector/Source/ClonerEffector/Private/Cloner/CEClonerComponent.cpp
--- a/Engine/Plugins/Experimental/ClonerEffector/Source/ClonerEffector/Private/Cloner/CEClonerComponent.cpp
+++ b/Engine/Plugins/Experimental/ClonerEffector/Source/ClonerEffector/Private/Cloner/CEClonerComponent.cpp
@@ -117,6 +117,28 @@ void UCEClonerComponent::PostLoad()
 }
 
 #if WITH_EDITOR
+void UCEClonerComponent::PostEditImport()
+{
+	SetAsset(nullptr);
+
+	Super::PostEditImport();
+
+	RegisterTicker();
+
+	ForceUpdateCloner();
+}
+
+void UCEClonerComponent::PostDuplicate(bool bInPIE)
+{
+	SetAsset(nullptr);
+
+	Super::PostDuplicate(bInPIE);
+
+	RegisterTicker();
+
+	ForceUpdateCloner();
+}
+
 void UCEClonerComponent::PostEditUndo()
 {
 	Super::PostEditUndo();
@@ -1265,6 +1287,14 @@ void UCEClonerComponent::ForceUpdateCloner()
 	OnLayoutNameChanged();
 }
 
+void UCEClonerComponent::OpenClonerSettings()
+{
+	if (const UCEClonerEffectorSettings* ClonerSettings = GetDefault<UCEClonerEffectorSettings>())
+	{
+		ClonerSettings->OpenSettings();
+	}
+}
+
 void UCEClonerComponent::CreateDefaultActorAttached()
 {
 	const UCEClonerEffectorSettings* ClonerEffectorSettings = GetDefault<UCEClonerEffectorSettings>();
@@ -1669,6 +1699,11 @@ UCEClonerLayoutBase* UCEClonerComponent::FindOrAddLayout(TSubclassOf<UCEClonerLa
 
 UCEClonerLayoutBase* UCEClonerComponent::FindOrAddLayout(FName InLayoutName)
 {
+	if (IsTemplate())
+	{
+		return nullptr;
+	}
+
 	UCEClonerSubsystem* Subsystem = UCEClonerSubsystem::Get();
 	if (!Subsystem)
 	{
diff --git a/Engine/Plugins/Experimental/ClonerEffector/Source/ClonerEffector/Private/Cloner/Layouts/CEClonerLayoutBase.cpp b/Engine/Plugins/Experimental/ClonerEffector/Source/ClonerEffector/Private/Cloner/Layouts/CEClonerLayoutBase.cpp
--- a/Engine/Plugins/Experimental/ClonerEffector/Source/ClonerEffector/Private/Cloner/Layouts/CEClonerLayoutBase.cpp
+++ b/Engine/Plugins/Experimental/ClonerEffector/Source/ClonerEffector/Private/Cloner/Layouts/CEClonerLayoutBase.cpp
@@ -85,9 +85,22 @@ void UCEClonerLayoutBase::LoadLayout()
 		return;
 	}
 
-	if (LayoutAssetPath.IsEmpty())
+	// System already cached and available with matching version
+	if (NiagaraSystem)
 	{
-		return;
+		if (CachedVersion == ECEClonerLayoutAssetVersion::LatestVersion)
+		{
+			UE_LOG(LogCEClonerLayoutBase, Verbose, TEXT("%s : Cloner layout %s using cached system V%i"), *GetClonerActor()->GetActorNameOrLabel(), *LayoutName.ToString(), CachedVersion)
+
+			CacheMeshRenderer();
+			OnSystemLoaded();
+			return;
+		}
+
+		UE_LOG(LogCEClonerLayoutBase, Warning, TEXT("%s : Cloner layout %s skipping cached system V%i due to version mismatch V%i"), *GetClonerActor()->GetActorNameOrLabel(), *LayoutName.ToString(), CachedVersion, ECEClonerLayoutAssetVersion::LatestVersion)
+
+		NiagaraSystem->MarkAsGarbage();
+		NiagaraSystem = nullptr;
 	}
 
 	const UCEClonerComponent* ClonerComponent = GetClonerComponent();
@@ -97,6 +110,11 @@ void UCEClonerLayoutBase::LoadLayout()
 		return;
 	}
 
+	if (LayoutAssetPath.IsEmpty())
+	{
+		return;
+	}
+
 	// Extract package path
 	FString MountedPath = LayoutAssetPath;
 	{
@@ -148,18 +166,14 @@ bool UCEClonerLayoutBase::UnloadLayout()
 	}
 
 	MeshRenderer->Meshes.Empty();
-	NiagaraSystem->RemoveFromRoot();
 #if WITH_EDITOR
+	MeshRenderer->OnMeshChanged();
+	MeshRenderer->OnChanged().Broadcast();
 	NiagaraSystem->KillAllActiveCompilations();
 #endif
-
-	// Prevent GC Leak
-	UPackage* Package = NiagaraSystem->GetPackage();
-	Package->ClearFlags(RF_Standalone);
-	Package->MarkAsGarbage();
+	NiagaraSystem->RemoveFromRoot();
 
 	MeshRenderer = nullptr;
-	NiagaraSystem = nullptr;
 
 	UE_LOG(LogCEClonerLayoutBase, Verbose, TEXT("%s : Cloner layout unloaded %s"), *GetClonerActor()->GetActorNameOrLabel(), *LayoutName.ToString())
 
@@ -280,6 +294,20 @@ void UCEClonerLayoutBase::PostEditImport()
 {
 	Super::PostEditImport();
 
+	// After cloner duplication in editor, niagara system should not be duplicated but still is,
+	// so look for it in outer chain otherwise it will trigger a world GC leak when switching level
+
+	TArray<UObject*> OwnedObjects;
+	GetObjectsWithOuter(this, OwnedObjects, false);
+
+	for (UObject* OwnedObject : OwnedObjects)
+	{
+		if (OwnedObject && OwnedObject->IsA<UNiagaraSystem>())
+		{
+			OwnedObject->MarkAsGarbage();
+		}
+	}
+
 	MarkLayoutDirty();
 }
 
@@ -306,22 +334,34 @@ void UCEClonerLayoutBase::OnSystemPackageLoaded(const FName& InName, UPackage* I
 	{
 		InPackage->SetFlags(RF_Transient);
 		NiagaraSystem->RemoveFromRoot();
-		NiagaraSystem->ClearFlags(RF_Standalone);
+		NiagaraSystem->ClearFlags(RF_Standalone | RF_Public | RF_Transient | RF_Transactional);
+
+		constexpr ERenameFlags RenameFlags = REN_NonTransactional | REN_DontCreateRedirectors;
+		if (NiagaraSystem->Rename(nullptr, this, RenameFlags))
+		{
+			CachedVersion = ECEClonerLayoutAssetVersion::LatestVersion;
+			CacheMeshRenderer();
+		}
 
-		CacheMeshRenderer();
+		InPackage->MarkAsGarbage();
 	}
 
+	OnSystemLoaded();
+}
+
+void UCEClonerLayoutBase::OnSystemLoaded()
+{
 	const bool bLayoutLoaded = IsLayoutLoaded();
 
 	if (bLayoutLoaded)
 	{
-		UE_LOG(LogCEClonerLayoutBase, Verbose, TEXT("%s : Cloner layout loaded %s - Template system %s - Package %s"), *GetClonerActor()->GetActorNameOrLabel(), *LayoutName.ToString(), *LayoutAssetPath, *InName.ToString())
+		UE_LOG(LogCEClonerLayoutBase, Verbose, TEXT("%s : Cloner layout loaded %s - Template system %s"), *GetClonerActor()->GetActorNameOrLabel(), *LayoutName.ToString(), *LayoutAssetPath)
 
 		OnLayoutLoaded();
 	}
 	else
 	{
-		UE_LOG(LogCEClonerLayoutBase, Warning, TEXT("%s : Cloner layout load failed %s - Template system %s - Package %s"), *GetClonerActor()->GetActorNameOrLabel(), *LayoutName.ToString(), *LayoutAssetPath, *InName.ToString())
+		UE_LOG(LogCEClonerLayoutBase, Warning, TEXT("%s : Cloner layout load failed %s - Template system %s"), *GetClonerActor()->GetActorNameOrLabel(), *LayoutName.ToString(), *LayoutAssetPath)
 	}
 
 	OnClonerLayoutLoadedDelegate.Broadcast(this, bLayoutLoaded);
diff --git a/Engine/Plugins/Experimental/ClonerEffector/Source/ClonerEffector/Public/Cloner/CEClonerComponent.h b/Engine/Plugins/Experimental/ClonerEffector/Source/ClonerEffector/Public/Cloner/CEClonerComponent.h
--- a/Engine/Plugins/Experimental/ClonerEffector/Source/ClonerEffector/Public/Cloner/CEClonerComponent.h
+++ b/Engine/Plugins/Experimental/ClonerEffector/Source/ClonerEffector/Public/Cloner/CEClonerComponent.h
@@ -167,6 +167,10 @@ public:
 	UFUNCTION(CallInEditor, Category="Cloner")
 	void ForceUpdateCloner();
 
+	/** Open project settings for cloner */
+	UFUNCTION(CallInEditor, Category="Utilities")
+	void OpenClonerSettings();
+
 	/** This will create a new default actor attached to this cloner if nothing is attached to this cloner */
 	UFUNCTION(CallInEditor, Category="Utilities")
 	void CreateDefaultActorAttached();
@@ -247,6 +251,8 @@ protected:
 	virtual void PostInitProperties() override;
 	virtual void PostLoad() override;
 #if WITH_EDITOR
+	virtual void PostEditImport() override;
+	virtual void PostDuplicate(bool bInPIE) override;
 	virtual void PostEditUndo() override;
 	virtual void PostEditChangeProperty(FPropertyChangedEvent& InPropertyChangedEvent) override;
 #endif
diff --git a/Engine/Plugins/Experimental/ClonerEffector/Source/ClonerEffector/Public/Cloner/Layouts/CEClonerLayoutBase.h b/Engine/Plugins/Experimental/ClonerEffector/Source/ClonerEffector/Public/Cloner/Layouts/CEClonerLayoutBase.h
--- a/Engine/Plugins/Experimental/ClonerEffector/Source/ClonerEffector/Public/Cloner/Layouts/CEClonerLayoutBase.h
+++ b/Engine/Plugins/Experimental/ClonerEffector/Source/ClonerEffector/Public/Cloner/Layouts/CEClonerLayoutBase.h
@@ -12,6 +12,16 @@ class UCEClonerExtensionBase;
 class UNiagaraMeshRendererProperties;
 class UNiagaraSystem;
 
+UENUM()
+enum class ECEClonerLayoutAssetVersion : uint8
+{
+	PreVersioning = 0,
+	/** Niagara System is now cached to improve loading time */
+	CachedVersion,
+	LatestVersionPlusOne,
+	LatestVersion = LatestVersionPlusOne - 1
+};
+
 /**
  * Base class for layouts available in the cloner actor
  * Steps to add a new layout :
@@ -141,6 +151,7 @@ private:
 
 	/** Called when the system package was async loaded */
 	void OnSystemPackageLoaded(const FName& InName, UPackage* InPackage, EAsyncLoadingResult::Type InResult);
+	void OnSystemLoaded();
 
 	/** Finds and cache first mesh renderer in emitter */
 	void CacheMeshRenderer();
@@ -161,15 +172,21 @@ private:
 	UPROPERTY(Transient)
 	FString LayoutAssetPath;
 
-	/** Niagara system representing this layout */
-	UPROPERTY(Transient, DuplicateTransient)
+	/** Niagara system used for this layout, cached to save some time */
+	UPROPERTY(DuplicateTransient, TextExportTransient)
 	TObjectPtr<UNiagaraSystem> NiagaraSystem;
 
 	/** Mesh renderer in this niagara system */
-	UPROPERTY(Transient, DuplicateTransient)
+	UPROPERTY(Transient, DuplicateTransient, TextExportTransient)
 	TObjectPtr<UNiagaraMeshRendererProperties> MeshRenderer;
 
+	/** Version of the cached system for diffs */
+	UPROPERTY(DuplicateTransient, TextExportTransient)
+	ECEClonerLayoutAssetVersion CachedVersion = ECEClonerLayoutAssetVersion::PreVersioning;
+
+	/** Id for the load request initiated */
 	int32 LoadRequestIdentifier = INDEX_NONE;
 
+	/** Status for this layout */
 	ECEClonerSystemStatus LayoutStatus = ECEClonerSystemStatus::UpToDate;
 };
\ No newline at end of file
