MotionDesign : ClonerEffector
- Add utility option to open cloner project settings
- Reduce motion ghosting turned on by default

diff --git a/Engine/Plugins/Experimental/ClonerEffector/Source/ClonerEffector/Private/Settings/CEClonerEffectorSettings.cpp b/Engine/Plugins/Experimental/ClonerEffector/Source/ClonerEffector/Private/Settings/CEClonerEffectorSettings.cpp
--- a/Engine/Plugins/Experimental/ClonerEffector/Source/ClonerEffector/Private/Settings/CEClonerEffectorSettings.cpp
+++ b/Engine/Plugins/Experimental/ClonerEffector/Source/ClonerEffector/Private/Settings/CEClonerEffectorSettings.cpp
@@ -8,6 +8,18 @@
 
 #if WITH_EDITOR
 #include "HAL/IConsoleManager.h"
+#include "ISettingsModule.h"
+#include "Modules/ModuleManager.h"
+#endif
+
+#if WITH_EDITOR
+void UCEClonerEffectorSettings::OpenSettings() const
+{
+	if (ISettingsModule* SettingsModule = FModuleManager::GetModulePtr<ISettingsModule>(TEXT("Settings")))
+	{
+		SettingsModule->ShowViewer(GetContainerName(), GetCategoryName(), GetSectionName());
+	}
+}
 #endif
 
 UCEClonerEffectorSettings::UCEClonerEffectorSettings()
@@ -42,7 +54,7 @@ void UCEClonerEffectorSettings::PostInitProperties()
 		CVarTSRShadingRejectionFlickeringPeriod->OnChangedDelegate().AddUObject(this, &UCEClonerEffectorSettings::OnTSRShadingRejectionFlickeringPeriodChanged);
 	}
 
-	OnTSRShadingRejectionFlickeringPeriodChanged(CVarTSRShadingRejectionFlickeringPeriod);
+	OnReduceMotionGhostingChanged();
 #endif
 }
 
diff --git a/Engine/Plugins/Experimental/ClonerEffector/Source/ClonerEffector/Private/Settings/CEClonerEffectorSettings.h b/Engine/Plugins/Experimental/ClonerEffector/Source/ClonerEffector/Private/Settings/CEClonerEffectorSettings.h
--- a/Engine/Plugins/Experimental/ClonerEffector/Source/ClonerEffector/Private/Settings/CEClonerEffectorSettings.h
+++ b/Engine/Plugins/Experimental/ClonerEffector/Source/ClonerEffector/Private/Settings/CEClonerEffectorSettings.h
@@ -45,6 +45,8 @@ public:
 	{
 		return bReduceMotionGhosting;
 	}
+
+	void OpenSettings() const;
 #endif
 
 protected:
@@ -80,11 +82,11 @@ protected:
 #if WITH_EDITORONLY_DATA
 	/** Reduces the r.TSR.ShadingRejection.Flickering.Period from 3 (default) to 1 if enabled to avoid ghosting artifacts when moving */
 	UPROPERTY(Config, EditAnywhere, Category="Cloner")
-	bool bReduceMotionGhosting = false;
+	bool bReduceMotionGhosting = true;
 #endif
 
 private:
-#if WITH_EDITOR
+#if WITH_EDITORONLY_DATA
 	void EnableNoFlicker();
 	void DisableNoFlicker();
 	bool IsNoFlickerEnabled() const;
@@ -96,6 +98,7 @@ private:
 	IConsoleVariable* CVarTSRShadingRejectionFlickeringPeriod = nullptr;
 
 	/** Previous value to restore it when disabled */
+	UPROPERTY()
 	TOptional<int32> PreviousCVarValue;
 #endif
 };
