MotionDesign : ClonerEffector
- Open cloner project settings window from cloner utilities section in details panel

diff --git a/Engine/Plugins/Experimental/ClonerEffector/Source/ClonerEffector/Private/Cloner/CEClonerComponent.cpp b/Engine/Plugins/Experimental/ClonerEffector/Source/ClonerEffector/Private/Cloner/CEClonerComponent.cpp
--- a/Engine/Plugins/Experimental/ClonerEffector/Source/ClonerEffector/Private/Cloner/CEClonerComponent.cpp
+++ b/Engine/Plugins/Experimental/ClonerEffector/Source/ClonerEffector/Private/Cloner/CEClonerComponent.cpp
@@ -1291,7 +1291,7 @@ void UCEClonerComponent::OpenClonerSettings()
 {
 	if (const UCEClonerEffectorSettings* ClonerSettings = GetDefault<UCEClonerEffectorSettings>())
 	{
-		ClonerSettings->OpenSettings();
+		ClonerSettings->OpenEditorSettingsWindow();
 	}
 }
 
diff --git a/Engine/Plugins/Experimental/ClonerEffector/Source/ClonerEffector/Private/Settings/CEClonerEffectorSettings.cpp b/Engine/Plugins/Experimental/ClonerEffector/Source/ClonerEffector/Private/Settings/CEClonerEffectorSettings.cpp
--- a/Engine/Plugins/Experimental/ClonerEffector/Source/ClonerEffector/Private/Settings/CEClonerEffectorSettings.cpp
+++ b/Engine/Plugins/Experimental/ClonerEffector/Source/ClonerEffector/Private/Settings/CEClonerEffectorSettings.cpp
@@ -12,16 +12,6 @@
 #include "Modules/ModuleManager.h"
 #endif
 
-#if WITH_EDITOR
-void UCEClonerEffectorSettings::OpenSettings() const
-{
-	if (ISettingsModule* SettingsModule = FModuleManager::GetModulePtr<ISettingsModule>(TEXT("Settings")))
-	{
-		SettingsModule->ShowViewer(GetContainerName(), GetCategoryName(), GetSectionName());
-	}
-}
-#endif
-
 UCEClonerEffectorSettings::UCEClonerEffectorSettings()
 {
 	CategoryName = TEXT("Motion Design");
@@ -59,6 +49,14 @@ void UCEClonerEffectorSettings::PostInitProperties()
 }
 
 #if WITH_EDITOR
+void UCEClonerEffectorSettings::OpenEditorSettingsWindow() const
+{
+	if (ISettingsModule* SettingsModule = FModuleManager::GetModulePtr<ISettingsModule>(TEXT("Settings")))
+	{
+		SettingsModule->ShowViewer(GetContainerName(), GetCategoryName(), GetSectionName());
+	}
+}
+
 void UCEClonerEffectorSettings::PostEditChangeProperty(FPropertyChangedEvent& InPropertyChangedEvent)
 {
 	if (InPropertyChangedEvent.GetMemberPropertyName() == GET_MEMBER_NAME_CHECKED(UCEClonerEffectorSettings, bReduceMotionGhosting))
diff --git a/Engine/Plugins/Experimental/ClonerEffector/Source/ClonerEffector/Private/Settings/CEClonerEffectorSettings.h b/Engine/Plugins/Experimental/ClonerEffector/Source/ClonerEffector/Private/Settings/CEClonerEffectorSettings.h
--- a/Engine/Plugins/Experimental/ClonerEffector/Source/ClonerEffector/Private/Settings/CEClonerEffectorSettings.h
+++ b/Engine/Plugins/Experimental/ClonerEffector/Source/ClonerEffector/Private/Settings/CEClonerEffectorSettings.h
@@ -46,7 +46,7 @@ public:
 		return bReduceMotionGhosting;
 	}
 
-	void OpenSettings() const;
+	void OpenEditorSettingsWindow() const;
 #endif
 
 protected:
