MotionDesign : ClonerEffector
- Fixed GC leak caused by pre versioning niagara system still being around when loading layout even if property was set transient

diff --git a/Engine/Plugins/Experimental/ClonerEffector/Source/ClonerEffector/Private/Cloner/Layouts/CEClonerLayoutBase.cpp b/Engine/Plugins/Experimental/ClonerEffector/Source/ClonerEffector/Private/Cloner/Layouts/CEClonerLayoutBase.cpp
--- a/Engine/Plugins/Experimental/ClonerEffector/Source/ClonerEffector/Private/Cloner/Layouts/CEClonerLayoutBase.cpp
+++ b/Engine/Plugins/Experimental/ClonerEffector/Source/ClonerEffector/Private/Cloner/Layouts/CEClonerLayoutBase.cpp
@@ -102,6 +102,10 @@ void UCEClonerLayoutBase::LoadLayout()
 		NiagaraSystem->MarkAsGarbage();
 		NiagaraSystem = nullptr;
 	}
+	else
+	{
+		CleanOwnedSystem();
+	}
 
 	const UCEClonerComponent* ClonerComponent = GetClonerComponent();
 
@@ -296,19 +300,21 @@ void UCEClonerLayoutBase::PostEditImport()
 
 	// After cloner duplication in editor, niagara system should not be duplicated but still is,
 	// so look for it in outer chain otherwise it will trigger a world GC leak when switching level
+	CleanOwnedSystem();
 
-	TArray<UObject*> OwnedObjects;
-	GetObjectsWithOuter(this, OwnedObjects, false);
+	MarkLayoutDirty();
+}
 
-	for (UObject* OwnedObject : OwnedObjects)
+void UCEClonerLayoutBase::PostLoad()
+{
+	Super::PostLoad();
+
+	if (CachedVersion == ECEClonerLayoutAssetVersion::PreVersioning)
 	{
-		if (OwnedObject && OwnedObject->IsA<UNiagaraSystem>())
-		{
-			OwnedObject->MarkAsGarbage();
-		}
+		// After cloner layout load, niagara system should not be loaded since property was transient pre versioning,
+		// so look for it in outer chain otherwise it will trigger a world GC leak when switching level
+		CleanOwnedSystem();
 	}
-
-	MarkLayoutDirty();
 }
 
 #if WITH_EDITOR
@@ -448,6 +454,21 @@ void UCEClonerLayoutBase::OnLevelCleanup()
 	UnbindCleanupDelegates();
 }
 
+void UCEClonerLayoutBase::CleanOwnedSystem() const
+{
+	TArray<UObject*> OwnedObjects;
+	GetObjectsWithOuter(this, OwnedObjects, false);
+
+	for (UObject* OwnedObject : OwnedObjects)
+	{
+		if (OwnedObject && OwnedObject->IsA<UNiagaraSystem>())
+		{
+			UE_LOG(LogCEClonerLayoutBase, Warning, TEXT("%s : Cloner layout %s cleaning owned system %s"), *GetClonerActor()->GetActorNameOrLabel(), *LayoutName.ToString(), *OwnedObject->GetName())
+			OwnedObject->MarkAsGarbage();
+		}
+	}
+}
+
 UCEClonerComponent* UCEClonerLayoutBase::GetClonerComponent() const
 {
 	return GetTypedOuter<UCEClonerComponent>();
diff --git a/Engine/Plugins/Experimental/ClonerEffector/Source/ClonerEffector/Public/Cloner/Layouts/CEClonerLayoutBase.h b/Engine/Plugins/Experimental/ClonerEffector/Source/ClonerEffector/Public/Cloner/Layouts/CEClonerLayoutBase.h
--- a/Engine/Plugins/Experimental/ClonerEffector/Source/ClonerEffector/Public/Cloner/Layouts/CEClonerLayoutBase.h
+++ b/Engine/Plugins/Experimental/ClonerEffector/Source/ClonerEffector/Public/Cloner/Layouts/CEClonerLayoutBase.h
@@ -124,6 +124,7 @@ public:
 protected:
 	//~ Begin UObject
 	virtual void PostEditImport() override;
+	virtual void PostLoad() override;
 #if WITH_EDITOR
 	virtual void PostEditUndo() override;
 #endif
@@ -164,6 +165,9 @@ private:
 	void OnWorldCleanup(UWorld* InWorld, bool bInSessionEnded, bool bInCleanupResources);
 	void OnLevelCleanup();
 
+	/** Marks niagara system owned by this layout garbage to avoid GC leak */
+	void CleanOwnedSystem() const;
+
 	/** Layout name to display in layout options */
 	UPROPERTY(Transient)
 	FName LayoutName;
