Fixed some TObjectPtr-from-reference constructions.

diff --git a/Engine/Plugins/Experimental/ClonerEffector/Source/ClonerEffector/Private/Cloner/CEClonerComponent.cpp b/Engine/Plugins/Experimental/ClonerEffector/Source/ClonerEffector/Private/Cloner/CEClonerComponent.cpp
--- a/Engine/Plugins/Experimental/ClonerEffector/Source/ClonerEffector/Private/Cloner/CEClonerComponent.cpp
+++ b/Engine/Plugins/Experimental/ClonerEffector/Source/ClonerEffector/Private/Cloner/CEClonerComponent.cpp
@@ -1926,7 +1926,7 @@ void UCEClonerComponent::OnActiveLayoutChanged()
 		InactiveExtension->DeactivateExtension();
 	}
 
-	ActiveExtensions.StableSort([](const TObjectPtr<UCEClonerExtensionBase>& InExtensionA, const TObjectPtr<UCEClonerExtensionBase>& InExtensionB)
+	Algo::StableSort(ActiveExtensions, [](const TObjectPtr<UCEClonerExtensionBase>& InExtensionA, const TObjectPtr<UCEClonerExtensionBase>& InExtensionB)
 	{
 		return InExtensionA->GetExtensionPriority() > InExtensionB->GetExtensionPriority();
 	});
