MotionDesign : ClonerEffector
- Fix translucent clones flickering when overlapping

diff --git a/Engine/Plugins/Experimental/ClonerEffector/Source/ClonerEffector/Private/Cloner/Extensions/CEClonerMeshRendererExtension.cpp b/Engine/Plugins/Experimental/ClonerEffector/Source/ClonerEffector/Private/Cloner/Extensions/CEClonerMeshRendererExtension.cpp
--- a/Engine/Plugins/Experimental/ClonerEffector/Source/ClonerEffector/Private/Cloner/Extensions/CEClonerMeshRendererExtension.cpp
+++ b/Engine/Plugins/Experimental/ClonerEffector/Source/ClonerEffector/Private/Cloner/Extensions/CEClonerMeshRendererExtension.cpp
@@ -114,6 +114,17 @@ void UCEClonerMeshRendererExtension::SetOverrideMaterial(UMaterialInterface* InM
 	OnOverrideMaterialOptionsChanged();
 }
 
+void UCEClonerMeshRendererExtension::SetSortTranslucentParticles(bool bInSort)
+{
+	if (bSortTranslucentParticles == bInSort)
+	{
+		return;
+	}
+
+	bSortTranslucentParticles = bInSort;
+	OnOverrideMaterialOptionsChanged();	
+}
+
 void UCEClonerMeshRendererExtension::OnExtensionParametersChanged(UCEClonerComponent* InComponent)
 {
 	Super::OnExtensionParametersChanged(InComponent);
@@ -227,6 +238,7 @@ void UCEClonerMeshRendererExtension::OnClonerMeshesUpdated()
 
 	MeshRenderer->FacingMode = MeshFacingMode;
 	MeshRenderer->bCastShadows = bMeshCastShadows;
+	MeshRenderer->SortMode = bSortTranslucentParticles ? ENiagaraSortMode::ViewDepth : ENiagaraSortMode::None;
 
 	// Use default meshes if nothing is attached
 	if (ClonerComponent->GetAttachmentCount() == 0)
@@ -258,6 +270,7 @@ const TCEPropertyChangeDispatcher<UCEClonerMeshRendererExtension> UCEClonerMeshR
 	{ GET_MEMBER_NAME_CHECKED(UCEClonerMeshRendererExtension, bUseOverrideMaterial), &UCEClonerMeshRendererExtension::OnOverrideMaterialOptionsChanged },
 	{ GET_MEMBER_NAME_CHECKED(UCEClonerMeshRendererExtension, OverrideMaterial), &UCEClonerMeshRendererExtension::OnOverrideMaterialOptionsChanged },
 	{ GET_MEMBER_NAME_CHECKED(UCEClonerMeshRendererExtension, bVisualizeEffectors), &UCEClonerMeshRendererExtension::OnOverrideMaterialOptionsChanged },
+	{ GET_MEMBER_NAME_CHECKED(UCEClonerMeshRendererExtension, bSortTranslucentParticles), &UCEClonerMeshRendererExtension::OnOverrideMaterialOptionsChanged },
 };
 
 void UCEClonerMeshRendererExtension::PostEditChangeProperty(FPropertyChangedEvent& InPropertyChangedEvent)
diff --git a/Engine/Plugins/Experimental/ClonerEffector/Source/ClonerEffector/Public/Cloner/Extensions/CEClonerMeshRendererExtension.h b/Engine/Plugins/Experimental/ClonerEffector/Source/ClonerEffector/Public/Cloner/Extensions/CEClonerMeshRendererExtension.h
--- a/Engine/Plugins/Experimental/ClonerEffector/Source/ClonerEffector/Public/Cloner/Extensions/CEClonerMeshRendererExtension.h
+++ b/Engine/Plugins/Experimental/ClonerEffector/Source/ClonerEffector/Public/Cloner/Extensions/CEClonerMeshRendererExtension.h
@@ -84,6 +84,15 @@ public:
 		return OverrideMaterial;
 	}
 
+	UFUNCTION(BlueprintCallable, Category="Cloner")
+	CLONEREFFECTOR_API void SetSortTranslucentParticles(bool bInSort);
+
+	UFUNCTION(BlueprintPure, Category="Cloner")
+	bool GetSortTranslucentParticles() const
+	{
+		return bSortTranslucentParticles;
+	}
+
 protected:
 	//~ Begin UObject
 #if WITH_EDITOR
@@ -132,6 +141,10 @@ protected:
 	UPROPERTY(EditInstanceOnly, Setter, Getter, Category="Renderer", meta=(EditCondition="bUseOverrideMaterial", EditConditionHides))
 	TObjectPtr<UMaterialInterface> OverrideMaterial;
 
+	/** Sort particles by depth when it has a translucent material, this will avoid flickering artifacts from appearing */
+	UPROPERTY(EditInstanceOnly, Setter="SetSortTranslucentParticles", Getter="GetSortTranslucentParticles", Category="Renderer")
+	bool bSortTranslucentParticles = true;
+
 private:
 #if WITH_EDITOR
 	/** Used for PECP */
