MotionDesign : ClonerEffector
- Fix forces would still affect clones when toggled off

diff --git a/Engine/Plugins/Experimental/ClonerEffector/Source/ClonerEffector/Private/Effector/CEEffectorComponent.cpp b/Engine/Plugins/Experimental/ClonerEffector/Source/ClonerEffector/Private/Effector/CEEffectorComponent.cpp
--- a/Engine/Plugins/Experimental/ClonerEffector/Source/ClonerEffector/Private/Effector/CEEffectorComponent.cpp
+++ b/Engine/Plugins/Experimental/ClonerEffector/Source/ClonerEffector/Private/Effector/CEEffectorComponent.cpp
@@ -262,10 +262,7 @@ void UCEEffectorComponent::RequestClonerUpdate(bool bInImmediate)
 		{
 			if (UCEClonerEffectorExtension* ClonerExtension = ClonerExtensionWeak.Get())
 			{
-				if (UCEClonerComponent* ClonerComponent = ClonerExtension->GetClonerComponent())
-				{
-					ClonerComponent->RequestClonerUpdate(bInImmediate);
-				}
+				ClonerExtension->MarkExtensionDirty(bInImmediate);
 			}
 		}
 	}
