MotionDesign : ClonerEffector
- Swap visualizer color when effector type is inverted

diff --git a/Engine/Plugins/Experimental/ClonerEffector/Source/ClonerEffector/Private/Effector/CEEffectorComponent.cpp b/Engine/Plugins/Experimental/ClonerEffector/Source/ClonerEffector/Private/Effector/CEEffectorComponent.cpp
--- a/Engine/Plugins/Experimental/ClonerEffector/Source/ClonerEffector/Private/Effector/CEEffectorComponent.cpp
+++ b/Engine/Plugins/Experimental/ClonerEffector/Source/ClonerEffector/Private/Effector/CEEffectorComponent.cpp
@@ -52,13 +52,6 @@ UCEEffectorComponent::UCEEffectorComponent()
 
 		UCEEffectorSubsystem::OnEffectorSetEnabled().AddUObject(this, &UCEEffectorComponent::OnEffectorSetEnabled);
 		TransformUpdated.AddUObject(this, &UCEEffectorComponent::OnTransformUpdated);
-
-#if WITH_EDITOR
-		if (UCEClonerEffectorSettings* ClonerEffectorSettings = GetMutableDefault<UCEClonerEffectorSettings>())
-		{
-			ClonerEffectorSettings->OnSettingChanged().AddUObject(this, &UCEEffectorComponent::OnEffectorDeveloperSettingsChanged);
-		}
-#endif
 	}
 }
 
@@ -454,6 +447,18 @@ void UCEEffectorComponent::UpdateVisualizer(int32 InVisualizerIndex, TFunctionRe
 		MeshComponent->SetMaterial(0, VisualizerMaterial);
 	}
 }
+
+void UCEEffectorComponent::SetVisualizerColor(int32 InVisualizerIndex, const FLinearColor& InColor)
+{
+	UMaterialInstanceDynamic* VisualizerMaterial = VisualizerMaterialsWeak.IsValidIndex(InVisualizerIndex) ? VisualizerMaterialsWeak[InVisualizerIndex].Get() : nullptr;
+
+	if (!IsValid(VisualizerMaterial))
+	{
+		return;
+	}
+
+	VisualizerMaterial->SetVectorParameterValue(VisualizerColorName, InColor);
+}
 #endif
 
 void UCEEffectorComponent::SetColor(const FLinearColor& InColor)
@@ -683,38 +688,8 @@ void UCEEffectorComponent::OnExtensionsChanged()
 }
 
 #if WITH_EDITOR
-void UCEEffectorComponent::OnEffectorDeveloperSettingsChanged(UObject* InSettings, FPropertyChangedEvent& InEvent)
-{
-	OnVisualizerOptionsChanged();
-}
-
 void UCEEffectorComponent::OnVisualizerOptionsChanged()
 {
-	const UCEClonerEffectorSettings* ClonerEffectorSettings = GetDefault<UCEClonerEffectorSettings>();
-
-	if (!ClonerEffectorSettings)
-	{
-		return;
-	}
-
-	int32 MaterialIndex = 0;
-	if (VisualizerMaterialsWeak.IsValidIndex(MaterialIndex))
-	{
-		if (UMaterialInstanceDynamic* VisualizerMaterials = VisualizerMaterialsWeak[MaterialIndex].Get())
-		{
-			VisualizerMaterials->SetVectorParameterValue(VisualizerColorName, ClonerEffectorSettings->GetVisualizerInnerColor());
-		}
-	}
-
-	MaterialIndex = 1;
-	if (VisualizerMaterialsWeak.IsValidIndex(MaterialIndex))
-	{
-		if (UMaterialInstanceDynamic* VisualizerMaterials = VisualizerMaterialsWeak[MaterialIndex].Get())
-		{
-			VisualizerMaterials->SetVectorParameterValue(VisualizerColorName, ClonerEffectorSettings->GetVisualizerOuterColor());
-		}
-	}
-
 	for (const TWeakObjectPtr<UDynamicMeshComponent>& VisualizerComponentWeak : VisualizerComponentsWeak)
 	{
 		if (UDynamicMeshComponent* VisualizerComponent = VisualizerComponentWeak.Get())
diff --git a/Engine/Plugins/Experimental/ClonerEffector/Source/ClonerEffector/Private/Effector/Types/CEEffectorBoundType.cpp b/Engine/Plugins/Experimental/ClonerEffector/Source/ClonerEffector/Private/Effector/Types/CEEffectorBoundType.cpp
--- a/Engine/Plugins/Experimental/ClonerEffector/Source/ClonerEffector/Private/Effector/Types/CEEffectorBoundType.cpp
+++ b/Engine/Plugins/Experimental/ClonerEffector/Source/ClonerEffector/Private/Effector/Types/CEEffectorBoundType.cpp
@@ -3,6 +3,7 @@
 #include "Effector/Types/CEEffectorBoundType.h"
 
 #include "Effector/CEEffectorComponent.h"
+#include "Settings/CEClonerEffectorSettings.h"
 
 void UCEEffectorBoundType::SetInvertType(bool bInInvert)
 {
@@ -33,6 +34,38 @@ void UCEEffectorBoundType::OnExtensionParametersChanged(UCEEffectorComponent* In
 	FCEClonerEffectorChannelData& ChannelData = InComponent->GetChannelData();
 	ChannelData.Easing = Easing;
 	ChannelData.Magnitude = bInvertType ? -InComponent->GetMagnitude() : InComponent->GetMagnitude();
+
+#if WITH_EDITOR
+	if (const UCEClonerEffectorSettings* ClonerEffectorSettings = GetDefault<UCEClonerEffectorSettings>())
+	{
+		InComponent->SetVisualizerColor(0, bInvertType ? ClonerEffectorSettings->GetVisualizerOuterColor() : ClonerEffectorSettings->GetVisualizerInnerColor());
+		InComponent->SetVisualizerColor(1, bInvertType ? ClonerEffectorSettings->GetVisualizerInnerColor() : ClonerEffectorSettings->GetVisualizerOuterColor());
+	}
+#endif
+}
+
+void UCEEffectorBoundType::OnExtensionActivated()
+{
+	Super::OnExtensionActivated();
+	
+#if WITH_EDITOR
+	if (UCEClonerEffectorSettings* ClonerEffectorSettings = GetMutableDefault<UCEClonerEffectorSettings>())
+	{
+		ClonerEffectorSettings->OnSettingChanged().AddUObject(this, &UCEEffectorBoundType::OnEffectorDeveloperSettingsChanged);
+	}
+#endif
+}
+
+void UCEEffectorBoundType::OnExtensionDeactivated()
+{
+	Super::OnExtensionDeactivated();
+
+#if WITH_EDITOR
+	if (UCEClonerEffectorSettings* ClonerEffectorSettings = GetMutableDefault<UCEClonerEffectorSettings>())
+	{
+		ClonerEffectorSettings->OnSettingChanged().RemoveAll(this);
+	}
+#endif
 }
 
 #if WITH_EDITOR
@@ -48,4 +81,9 @@ void UCEEffectorBoundType::PostEditChangeProperty(FPropertyChangedEvent& InPrope
 
 	PropertyChangeDispatcher.OnPropertyChanged(this, InPropertyChangedEvent);
 }
+
+void UCEEffectorBoundType::OnEffectorDeveloperSettingsChanged(UObject* InSettings, FPropertyChangedEvent& InEvent)
+{
+	UpdateExtensionParameters();	
+}
 #endif
diff --git a/Engine/Plugins/Experimental/ClonerEffector/Source/ClonerEffector/Public/Effector/CEEffectorComponent.h b/Engine/Plugins/Experimental/ClonerEffector/Source/ClonerEffector/Public/Effector/CEEffectorComponent.h
--- a/Engine/Plugins/Experimental/ClonerEffector/Source/ClonerEffector/Public/Effector/CEEffectorComponent.h
+++ b/Engine/Plugins/Experimental/ClonerEffector/Source/ClonerEffector/Public/Effector/CEEffectorComponent.h
@@ -118,6 +118,8 @@ public:
 
 	/** Updates a specific mesh visualizers based on its index */
 	void UpdateVisualizer(int32 InVisualizerIndex, TFunctionRef<void(UDynamicMesh*)> InMeshFunction) const;
+
+	void SetVisualizerColor(int32 InVisualizerIndex, const FLinearColor& InColor);
 #endif
 
 	/** Get linked cloners */
@@ -258,7 +260,6 @@ private:
 	void OnExtensionsChanged();
 
 #if WITH_EDITOR
-	void OnEffectorDeveloperSettingsChanged(UObject* InSettings, FPropertyChangedEvent& InEvent);
 	void OnVisualizerOptionsChanged();
 #endif
 
diff --git a/Engine/Plugins/Experimental/ClonerEffector/Source/ClonerEffector/Public/Effector/Types/CEEffectorBoundType.h b/Engine/Plugins/Experimental/ClonerEffector/Source/ClonerEffector/Public/Effector/Types/CEEffectorBoundType.h
--- a/Engine/Plugins/Experimental/ClonerEffector/Source/ClonerEffector/Public/Effector/Types/CEEffectorBoundType.h
+++ b/Engine/Plugins/Experimental/ClonerEffector/Source/ClonerEffector/Public/Effector/Types/CEEffectorBoundType.h
@@ -44,6 +44,8 @@ public:
 protected:
 	//~ Begin UCEEffectorExtensionBase
 	virtual void OnExtensionParametersChanged(UCEEffectorComponent* InComponent) override;
+	virtual void OnExtensionActivated() override;
+	virtual void OnExtensionDeactivated() override;
 	//~ End UCEEffectorExtensionBase
 
 	//~ Begin UObject
@@ -51,6 +53,10 @@ protected:
 	virtual void PostEditChangeProperty(FPropertyChangedEvent& InPropertyChangedEvent) override;
 #endif
 	//~ End UObject
+	
+#if WITH_EDITOR
+	void OnEffectorDeveloperSettingsChanged(UObject* InSettings, FPropertyChangedEvent& InEvent);
+#endif
 
 	/** Invert the type effect, instead of affecting the inside of a zone, will affect the outside */
 	UPROPERTY(EditInstanceOnly, Setter="SetInvertType", Getter="GetInvertType", Category="Shape")
