MotionDesign : ClonerEffector
- Preview collision grid size and grid resolution in editor

#virtualized

diff --git a/Engine/Plugins/Experimental/ClonerEffector/Source/ClonerEffector/Private/Cloner/CEClonerComponent.cpp b/Engine/Plugins/Experimental/ClonerEffector/Source/ClonerEffector/Private/Cloner/CEClonerComponent.cpp
--- a/Engine/Plugins/Experimental/ClonerEffector/Source/ClonerEffector/Private/Cloner/CEClonerComponent.cpp
+++ b/Engine/Plugins/Experimental/ClonerEffector/Source/ClonerEffector/Private/Cloner/CEClonerComponent.cpp
@@ -1108,7 +1108,7 @@ void UCEClonerComponent::InitializeCloner()
 	RegisterTicker();
 
 	// Load layout after registering ticker to let attachment tree update, then layout rendering will occur with up to date attachment data
-	OnLayoutNameChanged();
+	OnEnabledChanged();
 
 	OnClonerInitializedDelegate.Broadcast(this);
 }
diff --git a/Engine/Plugins/Experimental/ClonerEffector/Source/ClonerEffector/Private/Cloner/Extensions/CEClonerCollisionExtension.cpp b/Engine/Plugins/Experimental/ClonerEffector/Source/ClonerEffector/Private/Cloner/Extensions/CEClonerCollisionExtension.cpp
--- a/Engine/Plugins/Experimental/ClonerEffector/Source/ClonerEffector/Private/Cloner/Extensions/CEClonerCollisionExtension.cpp
+++ b/Engine/Plugins/Experimental/ClonerEffector/Source/ClonerEffector/Private/Cloner/Extensions/CEClonerCollisionExtension.cpp
@@ -3,6 +3,8 @@
 #include "Cloner/Extensions/CEClonerCollisionExtension.h"
 
 #include "Cloner/CEClonerComponent.h"
+#include "Components/StaticMeshComponent.h"
+#include "Materials/MaterialInstanceDynamic.h"
 #include "NiagaraDataInterfaceArrayFloat.h"
 #include "NiagaraMeshRendererProperties.h"
 #include "NiagaraSystem.h"
@@ -171,6 +173,16 @@ void UCEClonerCollisionExtension::OnExtensionParametersChanged(UCEClonerComponen
 			CollisionRadiiDI->GetArrayReference() = CollisionRadii;
 		}
 	}
+
+#if WITH_EDITOR
+	if (InitVisualizerComponent())
+	{
+		CollisionVisualizerComponent->SetVisibility(bParticleCollisionEnabled && bPreviewCollisionGrid);
+		// Divide by 100 because cube size is 100x100x100
+		CollisionVisualizerComponent->SetWorldScale3D(CollisionGridSize / 100.f);
+		CollisionVisualizerMaterial->SetScalarParameterValue(TEXT("GridTile"), CollisionGridResolution);
+	}
+#endif
 }
 
 void UCEClonerCollisionExtension::OnClonerMeshesUpdated()
@@ -180,7 +192,98 @@ void UCEClonerCollisionExtension::OnClonerMeshesUpdated()
 	MarkExtensionDirty();
 }
 
+void UCEClonerCollisionExtension::OnExtensionDeactivated()
+{
+	Super::OnExtensionDeactivated();
+
 #if WITH_EDITOR
+	DestroyVisualizerComponent();
+#endif
+}
+
+#if WITH_EDITORONLY_DATA
+void UCEClonerCollisionExtension::SetPreviewCollisionGrid(bool bInPreview)
+{
+	if (bPreviewCollisionGrid == bInPreview)
+	{
+		return;
+	}
+
+	bPreviewCollisionGrid = bInPreview;
+	MarkExtensionDirty();
+}
+#endif
+
+#if WITH_EDITOR
+bool UCEClonerCollisionExtension::InitVisualizerComponent()
+{
+	if (CollisionVisualizerComponent)
+	{
+		return !!CollisionVisualizerMaterial;
+	}
+
+	UCEClonerComponent* ClonerComponent = GetClonerComponent();
+
+	if (!ClonerComponent)
+	{
+		return false;
+	}
+
+	AActor* ClonerActor = ClonerComponent->GetOwner();
+
+	if (!ClonerActor)
+	{
+		return false;
+	}
+
+	CollisionVisualizerComponent = NewObject<UStaticMeshComponent>(
+		ClonerActor
+		, UStaticMeshComponent::StaticClass()
+		, MakeUniqueObjectName(ClonerActor, UStaticMeshComponent::StaticClass(), TEXT("ClonerCollisionVisualizerComponent"))
+		, RF_Transient);
+	
+	CollisionVisualizerComponent->OnComponentCreated();
+	CollisionVisualizerComponent->SetupAttachment(ClonerComponent);
+	CollisionVisualizerComponent->RegisterComponent();
+	
+	CollisionVisualizerComponent->SetIsVisualizationComponent(true);
+	CollisionVisualizerComponent->SetHiddenInGame(true);
+	CollisionVisualizerComponent->SetCastShadow(false);
+
+	// Setup Mesh
+	if (UStaticMesh* CubeMesh = LoadObject<UStaticMesh>(nullptr, TEXT("/Script/Engine.StaticMesh'/Engine/BasicShapes/Cube.Cube'")))
+	{
+		CollisionVisualizerComponent->SetStaticMesh(CubeMesh);
+	}
+
+	// Setup Material
+	if (UMaterialInterface* VisualizerMaterial = LoadObject<UMaterialInterface>(nullptr, TEXT("/Script/Engine.Material'/ClonerEffector/Materials/M_ClonerGrid.M_ClonerGrid'")))
+	{
+		CollisionVisualizerMaterial = UMaterialInstanceDynamic::Create(
+			VisualizerMaterial,
+			CollisionVisualizerComponent
+		);
+		
+		CollisionVisualizerComponent->SetMaterial(0, CollisionVisualizerMaterial);
+	}
+
+	return !!CollisionVisualizerMaterial;
+}
+
+void UCEClonerCollisionExtension::DestroyVisualizerComponent()
+{
+	if (!CollisionVisualizerComponent)
+	{
+		return;
+	}
+
+	CollisionVisualizerComponent->DestroyComponent();
+	CollisionVisualizerComponent = nullptr;
+
+	CollisionVisualizerMaterial->MarkAsGarbage();
+	CollisionVisualizerMaterial = nullptr;
+}
+
 const TCEPropertyChangeDispatcher<UCEClonerCollisionExtension> UCEClonerCollisionExtension::PropertyChangeDispatcher =
 {
 	/** Collision */
@@ -194,6 +297,7 @@ const TCEPropertyChangeDispatcher<UCEClonerCollisionExtension> UCEClonerCollisio
 	{ GET_MEMBER_NAME_CHECKED(UCEClonerCollisionExtension, CollisionGridSize), &UCEClonerCollisionExtension::OnExtensionPropertyChanged },
 	{ GET_MEMBER_NAME_CHECKED(UCEClonerCollisionExtension, MassMin), &UCEClonerCollisionExtension::OnExtensionPropertyChanged },
 	{ GET_MEMBER_NAME_CHECKED(UCEClonerCollisionExtension, MassMax), &UCEClonerCollisionExtension::OnExtensionPropertyChanged },
+	{ GET_MEMBER_NAME_CHECKED(UCEClonerCollisionExtension, bPreviewCollisionGrid), &UCEClonerCollisionExtension::OnExtensionPropertyChanged },
 };
 
 void UCEClonerCollisionExtension::PostEditChangeProperty(FPropertyChangedEvent& InPropertyChangedEvent)
@@ -202,4 +306,4 @@ void UCEClonerCollisionExtension::PostEditChangeProperty(FPropertyChangedEvent&
 
 	PropertyChangeDispatcher.OnPropertyChanged(this, InPropertyChangedEvent);
 }
-#endif
\ No newline at end of file
+#endif
diff --git a/Engine/Plugins/Experimental/ClonerEffector/Source/ClonerEffector/Public/Cloner/Extensions/CEClonerCollisionExtension.h b/Engine/Plugins/Experimental/ClonerEffector/Source/ClonerEffector/Public/Cloner/Extensions/CEClonerCollisionExtension.h
--- a/Engine/Plugins/Experimental/ClonerEffector/Source/ClonerEffector/Public/Cloner/Extensions/CEClonerCollisionExtension.h
+++ b/Engine/Plugins/Experimental/ClonerEffector/Source/ClonerEffector/Public/Cloner/Extensions/CEClonerCollisionExtension.h
@@ -7,6 +7,9 @@
 #include "CEPropertyChangeDispatcher.h"
 #include "CEClonerCollisionExtension.generated.h"
 
+class UMaterialInstanceDynamic;
+class UStaticMeshComponent;
+
 /** Extension dealing with collisions and physics related options */
 UCLASS(MinimalAPI, BlueprintType, Within=CEClonerComponent, AutoExpandCategories=("Collisions"), meta=(Section="Physics", Priority=3))
 class UCEClonerCollisionExtension : public UCEClonerExtensionBase
@@ -108,6 +111,17 @@ public:
 		return MassMax;
 	}
 
+#if WITH_EDITORONLY_DATA
+	UFUNCTION(BlueprintCallable, Category="Cloner")
+	CLONEREFFECTOR_API void SetPreviewCollisionGrid(bool bInPreview);
+
+	UFUNCTION(BlueprintPure, Category="Cloner")
+	bool GetPreviewCollisionGrid() const
+	{
+		return bPreviewCollisionGrid;
+	}
+#endif
+
 protected:
 	//~ Begin UObject
 #if WITH_EDITOR
@@ -116,10 +130,16 @@ protected:
 	//~ End UObject
 
 	//~ Begin UCEClonerExtensionBase
+	virtual void OnExtensionDeactivated() override;
 	virtual void OnExtensionParametersChanged(UCEClonerComponent* InComponent) override;
 	virtual void OnClonerMeshesUpdated() override;
 	//~ End UCEClonerExtensionBase
 
+#if WITH_EDITOR
+	bool InitVisualizerComponent();
+	void DestroyVisualizerComponent();
+#endif
+
 	/**
 	 * Allow particles to react to surface by using the distance field,
 	 * ensure the mesh you want particle to collide with is tick enough or increase its distance field resolution scale in static mesh editor 
@@ -163,6 +183,19 @@ protected:
 	UPROPERTY(EditInstanceOnly, Setter, Getter, Category="Collisions", meta=(ClampMin="1"))
 	float MassMax = 2.f;
 
+#if WITH_EDITORONLY_DATA
+	/** Preview 3D neighbor grid used for collision */
+	UPROPERTY(EditInstanceOnly, Setter="SetPreviewCollisionGrid", Getter="GetPreviewCollisionGrid", Category="Collisions", meta=(EditCondition="bParticleCollisionEnabled", EditConditionHides))
+	bool bPreviewCollisionGrid = false;
+	
+	/** Used to preview collision grid in editor */
+	UPROPERTY(Transient, DuplicateTransient, TextExportTransient)
+	TObjectPtr<UStaticMeshComponent> CollisionVisualizerComponent;
+
+	UPROPERTY(Transient, DuplicateTransient, TextExportTransient)
+	TObjectPtr<UMaterialInstanceDynamic> CollisionVisualizerMaterial;
+#endif
+
 private:
 #if WITH_EDITOR
 	/** Used for PECP */
