MotionDesign : ClonerEffector
- Removing setter/getter for editoronly property

diff --git a/Engine/Plugins/Experimental/ClonerEffector/Source/ClonerEffector/Private/Cloner/Extensions/CEClonerCollisionExtension.cpp b/Engine/Plugins/Experimental/ClonerEffector/Source/ClonerEffector/Private/Cloner/Extensions/CEClonerCollisionExtension.cpp
--- a/Engine/Plugins/Experimental/ClonerEffector/Source/ClonerEffector/Private/Cloner/Extensions/CEClonerCollisionExtension.cpp
+++ b/Engine/Plugins/Experimental/ClonerEffector/Source/ClonerEffector/Private/Cloner/Extensions/CEClonerCollisionExtension.cpp
@@ -201,19 +201,6 @@ void UCEClonerCollisionExtension::OnExtensionDeactivated()
 #endif
 }
 
-#if WITH_EDITORONLY_DATA
-void UCEClonerCollisionExtension::SetPreviewCollisionGrid(bool bInPreview)
-{
-	if (bPreviewCollisionGrid == bInPreview)
-	{
-		return;
-	}
-
-	bPreviewCollisionGrid = bInPreview;
-	MarkExtensionDirty();
-}
-#endif
-
 #if WITH_EDITOR
 bool UCEClonerCollisionExtension::InitVisualizerComponent()
 {
diff --git a/Engine/Plugins/Experimental/ClonerEffector/Source/ClonerEffector/Public/Cloner/Extensions/CEClonerCollisionExtension.h b/Engine/Plugins/Experimental/ClonerEffector/Source/ClonerEffector/Public/Cloner/Extensions/CEClonerCollisionExtension.h
--- a/Engine/Plugins/Experimental/ClonerEffector/Source/ClonerEffector/Public/Cloner/Extensions/CEClonerCollisionExtension.h
+++ b/Engine/Plugins/Experimental/ClonerEffector/Source/ClonerEffector/Public/Cloner/Extensions/CEClonerCollisionExtension.h
@@ -111,17 +111,6 @@ public:
 		return MassMax;
 	}
 
-#if WITH_EDITORONLY_DATA
-	UFUNCTION(BlueprintCallable, Category="Cloner")
-	CLONEREFFECTOR_API void SetPreviewCollisionGrid(bool bInPreview);
-
-	UFUNCTION(BlueprintPure, Category="Cloner")
-	bool GetPreviewCollisionGrid() const
-	{
-		return bPreviewCollisionGrid;
-	}
-#endif
-
 protected:
 	//~ Begin UObject
 #if WITH_EDITOR
@@ -185,7 +174,7 @@ protected:
 
 #if WITH_EDITORONLY_DATA
 	/** Preview 3D neighbor grid used for collision */
-	UPROPERTY(EditInstanceOnly, Setter="SetPreviewCollisionGrid", Getter="GetPreviewCollisionGrid", Category="Collisions", meta=(EditCondition="bParticleCollisionEnabled", EditConditionHides))
+	UPROPERTY(EditInstanceOnly, Category="Collisions", meta=(EditCondition="bParticleCollisionEnabled", EditConditionHides))
 	bool bPreviewCollisionGrid = false;
 	
 	/** Used to preview collision grid in editor */
