MotionDesign : ClonerEffector
- Add cull effector to kill or hide clones when out of bounds, also allows to reduce clones scale to dissapear smoothly

#virtualized

diff --git a/Engine/Plugins/Experimental/ClonerEffector/Source/ClonerEffector/Private/Effector/Modes/CEEffectorCullMode.cpp b/Engine/Plugins/Experimental/ClonerEffector/Source/ClonerEffector/Private/Effector/Modes/CEEffectorCullMode.cpp
new file mode 100644
--- /dev/null
+++ b/Engine/Plugins/Experimental/ClonerEffector/Source/ClonerEffector/Private/Effector/Modes/CEEffectorCullMode.cpp
@@ -0,0 +1,78 @@
+// Copyright Epic Games, Inc. All Rights Reserved.
+
+#include "Effector/Modes/CEEffectorCullMode.h"
+
+#include "Effector/CEEffectorComponent.h"
+#include "NiagaraTypes.h"
+
+void UCEEffectorCullMode::SetBehavior(ECEEffectorCullModeBehavior InBehavior)
+{
+	if (Behavior == InBehavior)
+	{
+		return;
+	}
+
+	Behavior = InBehavior;
+	OnBehaviorChanged();
+}
+
+void UCEEffectorCullMode::SetScale(const FVector& InScale)
+{
+	const FVector NewScale = InScale.ComponentMax(FVector(UE_KINDA_SMALL_NUMBER));
+
+	if (NewScale.Equals(Scale))
+	{
+		return;
+	}
+
+	Scale = NewScale;
+	UpdateExtensionParameters();
+}
+
+void UCEEffectorCullMode::OnExtensionParametersChanged(UCEEffectorComponent* InComponent)
+{
+	Super::OnExtensionParametersChanged(InComponent);
+
+	FCEClonerEffectorChannelData& ChannelData = InComponent->GetChannelData();
+
+	ChannelData.LocationDelta = FVector::ZeroVector;
+	ChannelData.RotationDelta = FVector::ZeroVector;
+	ChannelData.ScaleDelta = Scale.ComponentMax(FVector(UE_KINDA_SMALL_NUMBER));
+	ChannelData.Frequency = static_cast<float>(Behavior);
+}
+
+void UCEEffectorCullMode::PostInitProperties()
+{
+	Super::PostInitProperties();
+
+	if (HasAnyFlags(RF_ClassDefaultObject))
+	{
+		// Register new type def for niagara
+
+		constexpr ENiagaraTypeRegistryFlags MeshFlags =
+			ENiagaraTypeRegistryFlags::AllowAnyVariable |
+			ENiagaraTypeRegistryFlags::AllowParameter;
+
+		FNiagaraTypeRegistry::Register(FNiagaraTypeDefinition(StaticEnum<ECEEffectorCullModeBehavior>()), MeshFlags);
+	}
+}
+
+void UCEEffectorCullMode::OnBehaviorChanged()
+{
+	UpdateExtensionParameters(/** UpdateCloner */true);		
+}
+
+#if WITH_EDITOR
+const TCEPropertyChangeDispatcher<UCEEffectorCullMode> UCEEffectorCullMode::PropertyChangeDispatcher =
+{
+	{ GET_MEMBER_NAME_CHECKED(UCEEffectorCullMode, Behavior), &UCEEffectorCullMode::OnBehaviorChanged },
+	{ GET_MEMBER_NAME_CHECKED(UCEEffectorCullMode, Scale), &UCEEffectorCullMode::OnExtensionPropertyChanged },
+};
+
+void UCEEffectorCullMode::PostEditChangeProperty(FPropertyChangedEvent& InPropertyChangedEvent)
+{
+	Super::PostEditChangeProperty(InPropertyChangedEvent);
+
+	PropertyChangeDispatcher.OnPropertyChanged(this, InPropertyChangedEvent);
+}
+#endif
diff --git a/Engine/Plugins/Experimental/ClonerEffector/Source/ClonerEffector/Public/CEClonerEffectorShared.h b/Engine/Plugins/Experimental/ClonerEffector/Source/ClonerEffector/Public/CEClonerEffectorShared.h
--- a/Engine/Plugins/Experimental/ClonerEffector/Source/ClonerEffector/Public/CEClonerEffectorShared.h
+++ b/Engine/Plugins/Experimental/ClonerEffector/Source/ClonerEffector/Public/CEClonerEffectorShared.h
@@ -161,7 +161,9 @@ enum class ECEClonerEffectorMode : uint8
 	/** Pushes clones apart based on a strength and direction */
 	Push,
 	/** Accumulate transform on clones based on their index */
-	Step
+	Step,
+	/** Culls clones based on a zone */
+	Cull
 };
 
 /** Enumerates the effector push direction available */
diff --git a/Engine/Plugins/Experimental/ClonerEffector/Source/ClonerEffector/Public/Cloner/Layouts/CEClonerLayoutBase.h b/Engine/Plugins/Experimental/ClonerEffector/Source/ClonerEffector/Public/Cloner/Layouts/CEClonerLayoutBase.h
--- a/Engine/Plugins/Experimental/ClonerEffector/Source/ClonerEffector/Public/Cloner/Layouts/CEClonerLayoutBase.h
+++ b/Engine/Plugins/Experimental/ClonerEffector/Source/ClonerEffector/Public/Cloner/Layouts/CEClonerLayoutBase.h
@@ -18,6 +18,8 @@ enum class ECEClonerLayoutAssetVersion : uint8
 	PreVersioning = 0,
 	/** Niagara System is now cached to improve loading time */
 	CachedVersion,
+	/** Cull clones feature added */
+	CullEffector,
 	LatestVersionPlusOne,
 	LatestVersion = LatestVersionPlusOne - 1
 };
diff --git a/Engine/Plugins/Experimental/ClonerEffector/Source/ClonerEffector/Public/Effector/Modes/CEEffectorCullMode.h b/Engine/Plugins/Experimental/ClonerEffector/Source/ClonerEffector/Public/Effector/Modes/CEEffectorCullMode.h
new file mode 100644
--- /dev/null
+++ b/Engine/Plugins/Experimental/ClonerEffector/Source/ClonerEffector/Public/Effector/Modes/CEEffectorCullMode.h
@@ -0,0 +1,76 @@
+// Copyright Epic Games, Inc. All Rights Reserved.
+
+#pragma once
+
+#include "CEClonerEffectorShared.h"
+#include "CEEffectorModeBase.h"
+#include "CEPropertyChangeDispatcher.h"
+#include "CEEffectorCullMode.generated.h"
+
+class UCEEffectorComponent;
+
+UENUM(BlueprintType)
+enum class ECEEffectorCullModeBehavior : uint8
+{
+	/** Kill clones based on effector shape, lower priority than other hide effector */
+	Kill,
+	/** Hide clones based on effector shape, higher priority than other kill effector */
+	Hide
+};
+
+UCLASS(MinimalAPI, BlueprintType, Within=CEEffectorComponent)
+class UCEEffectorCullMode : public UCEEffectorModeBase
+{
+	GENERATED_BODY()
+
+public:
+	UCEEffectorCullMode()
+		: UCEEffectorModeBase(TEXT("Cull"), static_cast<int32>(ECEClonerEffectorMode::Cull))
+	{}
+
+	UFUNCTION(BlueprintCallable, Category="Effector")
+	CLONEREFFECTOR_API void SetBehavior(ECEEffectorCullModeBehavior InBehavior);
+
+	UFUNCTION(BlueprintPure, Category="Effector")
+	ECEEffectorCullModeBehavior GetBehavior() const
+	{
+		return Behavior;
+	}
+	
+	UFUNCTION(BlueprintCallable, Category="Effector")
+	CLONEREFFECTOR_API void SetScale(const FVector& InScale);
+
+	UFUNCTION(BlueprintPure, Category="Effector")
+	const FVector& GetScale() const
+	{
+		return Scale;
+	}
+
+protected:
+	//~ Begin UCEEffectorNoiseMode
+	virtual void OnExtensionParametersChanged(UCEEffectorComponent* InComponent) override;
+	//~ End UCEEffectorNoiseMode
+
+	//~ Begin UObject
+	virtual void PostInitProperties() override;
+#if WITH_EDITOR
+	virtual void PostEditChangeProperty(FPropertyChangedEvent& InPropertyChangedEvent) override;
+#endif
+	//~ End UObject
+
+	void OnBehaviorChanged();
+
+	/** Expected behavior when culling clones */
+	UPROPERTY(EditInstanceOnly, Setter, Getter, Category="Mode")
+	ECEEffectorCullModeBehavior Behavior = ECEEffectorCullModeBehavior::Hide;
+	
+	/** Scale applied on affected clones before they are culled to slowly disappear */
+	UPROPERTY(EditInstanceOnly, Setter, Getter, Category="Mode", meta=(ClampMin="0", MotionDesignVectorWidget, AllowPreserveRatio="XYZ", Delta="0.01"))
+	FVector Scale = FVector::ZeroVector;
+	
+private:
+#if WITH_EDITOR
+	/** Used for PECP */
+	static const TCEPropertyChangeDispatcher<UCEEffectorCullMode> PropertyChangeDispatcher;
+#endif
+};
\ No newline at end of file
