Fix compilation error with ClonerEffector plugin following change to NiagaraTypeRegistry
#rnx
#tests compiled ClonerEffector module

diff --git a/Engine/Plugins/Experimental/ClonerEffector/Source/ClonerEffector/Private/Effector/Modes/CEEffectorCullMode.cpp b/Engine/Plugins/Experimental/ClonerEffector/Source/ClonerEffector/Private/Effector/Modes/CEEffectorCullMode.cpp
--- a/Engine/Plugins/Experimental/ClonerEffector/Source/ClonerEffector/Private/Effector/Modes/CEEffectorCullMode.cpp
+++ b/Engine/Plugins/Experimental/ClonerEffector/Source/ClonerEffector/Private/Effector/Modes/CEEffectorCullMode.cpp
@@ -4,6 +4,7 @@
 
 #include "Effector/CEEffectorComponent.h"
 #include "NiagaraTypes.h"
+#include "NiagaraTypeRegistry.h"
 
 void UCEEffectorCullMode::SetBehavior(ECEEffectorCullModeBehavior InBehavior)
 {
