MotionDesign : ClonerEffector
- Added delay effector extension to slow down clones mode transform and have a spring effect when going back to rest state, delay can be customized for peak and rest phase, frequency and falloff for spring effect can be adjusted as well

#virtualized

diff --git a/Engine/Plugins/Experimental/ClonerEffector/Source/ClonerEffector/Private/CEClonerEffectorShared.cpp b/Engine/Plugins/Experimental/ClonerEffector/Source/ClonerEffector/Private/CEClonerEffectorShared.cpp
--- a/Engine/Plugins/Experimental/ClonerEffector/Source/ClonerEffector/Private/CEClonerEffectorShared.cpp
+++ b/Engine/Plugins/Experimental/ClonerEffector/Source/ClonerEffector/Private/CEClonerEffectorShared.cpp
@@ -46,6 +46,10 @@ void FCEClonerEffectorChannelData::Write(UNiagaraDataChannelWriter* InWriter) co
 	InWriter->WriteFloat(FrequencyName, Identifier, Frequency);
 	InWriter->WriteVector(PanName, Identifier, Pan);
 	InWriter->WriteLinearColor(ColorName, Identifier, Color);
+	InWriter->WriteFloat(DelayInDurationName, Identifier, DelayInDuration);
+	InWriter->WriteFloat(DelayOutDurationName, Identifier, DelayOutDuration);
+	InWriter->WriteFloat(DelaySpringFrequencyName, Identifier, DelaySpringFrequency);
+	InWriter->WriteFloat(DelaySpringFalloffName, Identifier, DelaySpringFalloff);
 
 	/** Forces */
 	InWriter->WriteFloat(OrientationForceRateName, Identifier, OrientationForceRate);
@@ -87,6 +91,10 @@ void FCEClonerEffectorChannelData::Read(const UNiagaraDataChannelReader* InReade
 	Frequency = InReader->ReadFloat(FrequencyName, Identifier, bIsValid);
 	Pan = InReader->ReadVector(PanName, Identifier, bIsValid);
 	Color = InReader->ReadLinearColor(ColorName, Identifier, bIsValid);
+	DelayInDuration = InReader->ReadFloat(DelayInDurationName, Identifier, bIsValid);
+	DelayOutDuration = InReader->ReadFloat(DelayOutDurationName, Identifier, bIsValid);
+	DelaySpringFrequency = InReader->ReadFloat(DelaySpringFrequencyName, Identifier, bIsValid);
+	DelaySpringFalloff = InReader->ReadFloat(DelaySpringFalloffName, Identifier, bIsValid);
 
 	/** Forces */
 	OrientationForceRate = InReader->ReadFloat(OrientationForceRateName, Identifier, bIsValid);
diff --git a/Engine/Plugins/Experimental/ClonerEffector/Source/ClonerEffector/Private/Effector/Extensions/CEEffectorDelayExtension.cpp b/Engine/Plugins/Experimental/ClonerEffector/Source/ClonerEffector/Private/Effector/Extensions/CEEffectorDelayExtension.cpp
new file mode 100644
--- /dev/null
+++ b/Engine/Plugins/Experimental/ClonerEffector/Source/ClonerEffector/Private/Effector/Extensions/CEEffectorDelayExtension.cpp
@@ -0,0 +1,101 @@
+// Copyright Epic Games, Inc. All Rights Reserved.
+
+#include "Effector/Extensions/CEEffectorDelayExtension.h"
+
+#include "Effector/CEEffectorComponent.h"
+
+void UCEEffectorDelayExtension::SetDelayEnabled(bool bInEnabled)
+{
+	if (bDelayEnabled == bInEnabled)
+	{
+		return;
+	}
+
+	bDelayEnabled = bInEnabled;
+	UpdateExtensionParameters();
+}
+
+void UCEEffectorDelayExtension::SetDelayInDuration(float InDuration)
+{
+	InDuration = FMath::Max(InDuration, 0.f);
+	if (FMath::IsNearlyEqual(DelayInDuration, InDuration))
+	{
+		return;
+	}
+
+	DelayInDuration = InDuration;
+	UpdateExtensionParameters();
+}
+
+void UCEEffectorDelayExtension::SetDelayOutDuration(float InDuration)
+{
+	InDuration = FMath::Max(InDuration, 0.f);
+	if (FMath::IsNearlyEqual(DelayOutDuration, InDuration))
+	{
+		return;
+	}
+
+	DelayOutDuration = InDuration;
+	UpdateExtensionParameters();	
+}
+
+void UCEEffectorDelayExtension::SetDelaySpringFrequency(float InFrequency)
+{
+	InFrequency = FMath::Max(InFrequency, 1.f);
+	if (FMath::IsNearlyEqual(DelaySpringFrequency, InFrequency))
+	{
+		return;
+	}
+
+	DelaySpringFrequency = InFrequency;
+	UpdateExtensionParameters();
+}
+
+void UCEEffectorDelayExtension::SetDelaySpringFalloff(float InFalloff)
+{
+	InFalloff = FMath::Max(InFalloff, 1.f);
+	if (FMath::IsNearlyEqual(DelaySpringFalloff, InFalloff))
+	{
+		return;
+	}
+
+	DelaySpringFalloff = InFalloff;
+	UpdateExtensionParameters();
+}
+
+#if WITH_EDITOR
+const TCEPropertyChangeDispatcher<UCEEffectorDelayExtension> UCEEffectorDelayExtension::PropertyChangeDispatcher =
+{
+	{ GET_MEMBER_NAME_CHECKED(UCEEffectorDelayExtension, bDelayEnabled), &UCEEffectorDelayExtension::OnExtensionPropertyChanged },
+	{ GET_MEMBER_NAME_CHECKED(UCEEffectorDelayExtension, DelayInDuration), &UCEEffectorDelayExtension::OnExtensionPropertyChanged },
+	{ GET_MEMBER_NAME_CHECKED(UCEEffectorDelayExtension, DelaySpringFrequency), &UCEEffectorDelayExtension::OnExtensionPropertyChanged },
+	{ GET_MEMBER_NAME_CHECKED(UCEEffectorDelayExtension, DelaySpringFalloff), &UCEEffectorDelayExtension::OnExtensionPropertyChanged }
+};
+
+void UCEEffectorDelayExtension::PostEditChangeProperty(FPropertyChangedEvent& InPropertyChangedEvent)
+{
+	Super::PostEditChangeProperty(InPropertyChangedEvent);
+
+	PropertyChangeDispatcher.OnPropertyChanged(this, InPropertyChangedEvent);
+}
+#endif
+
+void UCEEffectorDelayExtension::OnExtensionParametersChanged(UCEEffectorComponent* InComponent)
+{
+	Super::OnExtensionParametersChanged(InComponent);
+
+	FCEClonerEffectorChannelData& ChannelData = InComponent->GetChannelData();
+
+	if (bDelayEnabled)
+	{
+		ChannelData.DelayInDuration = DelayInDuration;
+		ChannelData.DelayOutDuration = DelayOutDuration;
+		ChannelData.DelaySpringFrequency = DelaySpringFrequency * DelayOutDuration;
+		ChannelData.DelaySpringFalloff = DelaySpringFalloff;
+	}
+	else
+	{
+		ChannelData.DelayInDuration = 0.f;
+		ChannelData.DelayOutDuration = 0.f;
+	}
+}
diff --git a/Engine/Plugins/Experimental/ClonerEffector/Source/ClonerEffector/Public/CEClonerEffectorShared.h b/Engine/Plugins/Experimental/ClonerEffector/Source/ClonerEffector/Public/CEClonerEffectorShared.h
--- a/Engine/Plugins/Experimental/ClonerEffector/Source/ClonerEffector/Public/CEClonerEffectorShared.h
+++ b/Engine/Plugins/Experimental/ClonerEffector/Source/ClonerEffector/Public/CEClonerEffectorShared.h
@@ -389,6 +389,10 @@ struct FCEClonerEffectorChannelData
 	static constexpr const TCHAR* FrequencyName = TEXT("Frequency");
 	static constexpr const TCHAR* PanName = TEXT("Pan");
 	static constexpr const TCHAR* ColorName = TEXT("Color");
+	static constexpr const TCHAR* DelayInDurationName = TEXT("DelayInDuration");
+	static constexpr const TCHAR* DelayOutDurationName = TEXT("DelayOutDuration");
+	static constexpr const TCHAR* DelaySpringFrequencyName = TEXT("DelaySpringFrequency");
+	static constexpr const TCHAR* DelaySpringFalloffName = TEXT("DelaySpringFalloff");
 
 	/** Forces */
 	static constexpr const TCHAR* OrientationForceRateName = TEXT("OrientationForceRate");
@@ -426,6 +430,10 @@ struct FCEClonerEffectorChannelData
 	float Frequency = 1.f;
 	FVector Pan = FVector::ZeroVector;
 	FLinearColor Color = FLinearColor::Transparent;
+	float DelayInDuration = 0.f;
+	float DelayOutDuration = 0.f;
+	float DelaySpringFrequency = 0.f;
+	float DelaySpringFalloff = 0.f;
 
 	/** Forces parameters */
 	float OrientationForceRate = 0.f;
diff --git a/Engine/Plugins/Experimental/ClonerEffector/Source/ClonerEffector/Public/Cloner/CEClonerActor.h b/Engine/Plugins/Experimental/ClonerEffector/Source/ClonerEffector/Public/Cloner/CEClonerActor.h
--- a/Engine/Plugins/Experimental/ClonerEffector/Source/ClonerEffector/Public/Cloner/CEClonerActor.h
+++ b/Engine/Plugins/Experimental/ClonerEffector/Source/ClonerEffector/Public/Cloner/CEClonerActor.h
@@ -19,7 +19,7 @@ class USceneComponent;
 UCLASS(MinimalAPI
 	, BlueprintType
 	, AutoExpandCategories=(ClonerComponent)
-	, HideCategories=(Rendering,Replication,Collision,HLOD,Physics,Networking,Input,Actor,Cooking,LevelInstance,Streaming,DataLayers)
+	, HideCategories=(Rendering,Replication,Collision,HLOD,Physics,Networking,Input,Actor,Cooking,LevelInstance,Streaming,DataLayers,WorldPartition)
 	, DisplayName = "Motion Design Cloner Actor")
 class ACEClonerActor : public AActor
 {
diff --git a/Engine/Plugins/Experimental/ClonerEffector/Source/ClonerEffector/Public/Cloner/Layouts/CEClonerLayoutBase.h b/Engine/Plugins/Experimental/ClonerEffector/Source/ClonerEffector/Public/Cloner/Layouts/CEClonerLayoutBase.h
--- a/Engine/Plugins/Experimental/ClonerEffector/Source/ClonerEffector/Public/Cloner/Layouts/CEClonerLayoutBase.h
+++ b/Engine/Plugins/Experimental/ClonerEffector/Source/ClonerEffector/Public/Cloner/Layouts/CEClonerLayoutBase.h
@@ -20,6 +20,8 @@ enum class ECEClonerLayoutAssetVersion : uint8
 	CachedVersion,
 	/** Cull clones feature added */
 	CullEffector,
+	/** Delay clones feature added */
+	DelayEffector,
 	LatestVersionPlusOne,
 	LatestVersion = LatestVersionPlusOne - 1
 };
diff --git a/Engine/Plugins/Experimental/ClonerEffector/Source/ClonerEffector/Public/Effector/CEEffectorComponent.h b/Engine/Plugins/Experimental/ClonerEffector/Source/ClonerEffector/Public/Effector/CEEffectorComponent.h
--- a/Engine/Plugins/Experimental/ClonerEffector/Source/ClonerEffector/Public/Effector/CEEffectorComponent.h
+++ b/Engine/Plugins/Experimental/ClonerEffector/Source/ClonerEffector/Public/Effector/CEEffectorComponent.h
@@ -19,7 +19,7 @@ class UDynamicMeshComponent;
 UCLASS(
 	MinimalAPI
 	, BlueprintType
-	, HideCategories=(Tags,Activation,LOD,AssetUserData,Navigation,Rendering,Cooking,Physics)
+	, HideCategories=(Tags,Activation,LOD,AssetUserData,Navigation,Rendering,Cooking,Physics,WorldPartition)
 	, DisplayName = "Motion Design Effector Component")
 class UCEEffectorComponent : public USceneComponent
 {
diff --git a/Engine/Plugins/Experimental/ClonerEffector/Source/ClonerEffector/Public/Effector/Extensions/CEEffectorDelayExtension.h b/Engine/Plugins/Experimental/ClonerEffector/Source/ClonerEffector/Public/Effector/Extensions/CEEffectorDelayExtension.h
new file mode 100644
--- /dev/null
+++ b/Engine/Plugins/Experimental/ClonerEffector/Source/ClonerEffector/Public/Effector/Extensions/CEEffectorDelayExtension.h
@@ -0,0 +1,105 @@
+// Copyright Epic Games, Inc. All Rights Reserved.
+
+#pragma once
+
+#include "CEPropertyChangeDispatcher.h"
+#include "Effector/CEEffectorExtensionBase.h"
+#include "CEEffectorDelayExtension.generated.h"
+
+class UCEEffectorComponent;
+
+/** Effector delay extension to affect clones after a period of time with a spring effect */
+UCLASS(MinimalAPI, BlueprintType, Within=CEEffectorComponent, meta=(Section="Delay", Priority=3))
+class UCEEffectorDelayExtension : public UCEEffectorExtensionBase
+{
+	GENERATED_BODY()
+
+public:
+	UCEEffectorDelayExtension()
+		: UCEEffectorExtensionBase(
+			TEXT("Delay")
+		)
+	{}
+
+	UFUNCTION(BlueprintCallable, Category="Effector")
+	CLONEREFFECTOR_API void SetDelayEnabled(bool bInEnabled);
+
+	UFUNCTION(BlueprintPure, Category="Effector")
+	bool GetDelayEnabled() const
+	{
+		return bDelayEnabled;
+	}
+
+	UFUNCTION(BlueprintCallable, Category="Effector")
+	CLONEREFFECTOR_API void SetDelayInDuration(float InDuration);
+
+	UFUNCTION(BlueprintPure, Category="Effector")
+	float GetDelayInDuration() const
+	{
+		return DelayInDuration;
+	}
+
+	UFUNCTION(BlueprintCallable, Category="Effector")
+	CLONEREFFECTOR_API void SetDelayOutDuration(float InDuration);
+
+	UFUNCTION(BlueprintPure, Category="Effector")
+	float GetDelayOutDuration() const
+	{
+		return DelayOutDuration;
+	}
+
+	UFUNCTION(BlueprintCallable, Category="Effector")
+	CLONEREFFECTOR_API void SetDelaySpringFrequency(float InFrequency);
+
+	UFUNCTION(BlueprintPure, Category="Effector")
+	float GetDelaySpringFrequency() const
+	{
+		return DelaySpringFrequency;
+	}
+
+	UFUNCTION(BlueprintCallable, Category="Effector")
+	CLONEREFFECTOR_API void SetDelaySpringFalloff(float InFalloff);
+
+	UFUNCTION(BlueprintPure, Category="Effector")
+	float GetDelaySpringFalloff() const
+	{
+		return DelaySpringFalloff;
+	}
+
+protected:
+	//~ Begin UObject
+#if WITH_EDITOR
+	virtual void PostEditChangeProperty(FPropertyChangedEvent& InPropertyChangedEvent) override;
+#endif
+	//~ End UObject
+
+	//~ Begin UCEEffectorExtensionBase
+	virtual void OnExtensionParametersChanged(UCEEffectorComponent* InComponent) override;
+	//~ End UCEEffectorExtensionBase
+
+	/** Enable time delay for the mode effect */
+	UPROPERTY(EditInstanceOnly, Setter="SetDelayEnabled", Getter="GetDelayEnabled", Category="Delay")
+	bool bDelayEnabled = false;
+	
+	/** Time delay for the effect to reach its peak */
+	UPROPERTY(EditInstanceOnly, Setter, Getter, Category="Delay", meta=(Units=Seconds, ClampMin="0", EditCondition="bDelayEnabled", EditConditionHides))
+	float DelayInDuration = 1.f;
+
+	/** Time delay for the effect to go back to its rest state */
+	UPROPERTY(EditInstanceOnly, Setter, Getter, Category="Delay", meta=(Units=Seconds, ClampMin="0", EditCondition="bDelayEnabled", EditConditionHides))
+	float DelayOutDuration = 1.f;
+
+	/** Frequency of the delayed spring effect */
+	UPROPERTY(EditInstanceOnly, Setter, Getter, Category="Delay", meta=(ClampMin="1", EditCondition="bDelayEnabled", EditConditionHides))
+	float DelaySpringFrequency = 3.f;
+
+	/** Damping or decay of the delayed spring effect */
+	UPROPERTY(EditInstanceOnly, Setter, Getter, Category="Delay", meta=(ClampMin="1", EditCondition="bDelayEnabled", EditConditionHides))
+	float DelaySpringFalloff = 1.f;
+
+private:
+#if WITH_EDITOR
+	/** Used for PECP */
+	static const TCEPropertyChangeDispatcher<UCEEffectorDelayExtension> PropertyChangeDispatcher;
+#endif
+};
\ No newline at end of file
diff --git a/Engine/Plugins/Experimental/ClonerEffector/Source/ClonerEffector/Public/Effector/Modes/CEEffectorPushMode.h b/Engine/Plugins/Experimental/ClonerEffector/Source/ClonerEffector/Public/Effector/Modes/CEEffectorPushMode.h
--- a/Engine/Plugins/Experimental/ClonerEffector/Source/ClonerEffector/Public/Effector/Modes/CEEffectorPushMode.h
+++ b/Engine/Plugins/Experimental/ClonerEffector/Source/ClonerEffector/Public/Effector/Modes/CEEffectorPushMode.h
@@ -54,7 +54,7 @@ protected:
 
 	/** Relative direction computed for the push effect on each clone */
 	UPROPERTY(EditInstanceOnly, Setter, Getter, Category="Mode")
-	ECEClonerEffectorPushDirection PushDirection = ECEClonerEffectorPushDirection::Forward;
+	ECEClonerEffectorPushDirection PushDirection = ECEClonerEffectorPushDirection::Effector;
 
 private:
 #if WITH_EDITOR
