MotionDesign : ClonerEffector
- Fix effector unlinks itself from cloners in rundown due to not being initialized yet

diff --git a/Engine/Plugins/Experimental/ClonerEffector/Source/ClonerEffector/Private/Cloner/Extensions/CEClonerEffectorExtension.cpp b/Engine/Plugins/Experimental/ClonerEffector/Source/ClonerEffector/Private/Cloner/Extensions/CEClonerEffectorExtension.cpp
--- a/Engine/Plugins/Experimental/ClonerEffector/Source/ClonerEffector/Private/Cloner/Extensions/CEClonerEffectorExtension.cpp
+++ b/Engine/Plugins/Experimental/ClonerEffector/Source/ClonerEffector/Private/Cloner/Extensions/CEClonerEffectorExtension.cpp
@@ -95,6 +95,18 @@ void UCEClonerEffectorExtension::OnEffectorIdentifierChanged(UCEEffectorComponen
 	if (EffectorActorsWeak.Contains(InEffector->GetOwner()))
 	{
 		OnEffectorsChanged();
+
+		if (InOldIdentifier != INDEX_NONE && InNewIdentifier == INDEX_NONE)
+		{
+#if WITH_EDITOR
+			if (GUndo)
+			{
+				Modify();
+			}
+#endif
+
+			EffectorActorsWeak.Remove(InEffector->GetOwner());
+		}
 	}
 }
 
@@ -119,6 +131,7 @@ void UCEClonerEffectorExtension::OnEffectorsChanged()
 
 	// Remove duplicates
 	TSet<TWeakObjectPtr<UCEEffectorComponent>> SetEffectorsWeak;
+	SetEffectorsWeak.Reserve(EffectorActorsWeak.Num());
 
 	for (TArray<TWeakObjectPtr<AActor>>::TIterator It(EffectorActorsWeak); It; ++It)
 	{
@@ -134,18 +147,6 @@ void UCEClonerEffectorExtension::OnEffectorsChanged()
 					SetEffectorsWeak.Add(EffectorComponent);
 				}
 			}
-
-			if (SetEffectorsWeak.IsEmpty())
-			{
-#if WITH_EDITOR
-				if (GUndo)
-				{
-					Modify();
-				}
-#endif
-
-				It.RemoveCurrent();
-			}
 		}
 	}
 
