MotionDesign : ClonerEffector
- Fix cloner override material property is read only in details panel

diff --git a/Engine/Plugins/Experimental/ClonerEffector/Source/ClonerEffector/Public/Cloner/Extensions/CEClonerMeshRendererExtension.h b/Engine/Plugins/Experimental/ClonerEffector/Source/ClonerEffector/Public/Cloner/Extensions/CEClonerMeshRendererExtension.h
--- a/Engine/Plugins/Experimental/ClonerEffector/Source/ClonerEffector/Public/Cloner/Extensions/CEClonerMeshRendererExtension.h
+++ b/Engine/Plugins/Experimental/ClonerEffector/Source/ClonerEffector/Public/Cloner/Extensions/CEClonerMeshRendererExtension.h
@@ -138,7 +138,7 @@ protected:
 	bool bUseOverrideMaterial = false;
 
 	/** The override materials that will be set instead of meshes materials, bVisualizeEffectors must be disabled */
-	UPROPERTY(EditInstanceOnly, Setter, Getter, Category="Renderer", meta=(EditCondition="bUseOverrideMaterial", EditConditionHides))
+	UPROPERTY(EditAnywhere, Setter, Getter, Category="Renderer", meta=(EditCondition="bUseOverrideMaterial", EditConditionHides))
 	TObjectPtr<UMaterialInterface> OverrideMaterial;
 
 	/** Sort particles by depth when it has a translucent material, this will avoid flickering artifacts from appearing */
