MotionDesign : ClonerEffector
- Fix texture constraint threshold not working due to name mismatch with user parameter, using typo name for 5.5 to avoid resaving assets and overwriting 5.6 assets

diff --git a/Engine/Plugins/Experimental/ClonerEffector/Source/ClonerEffector/Private/Cloner/Extensions/CEClonerConstraintExtension.cpp b/Engine/Plugins/Experimental/ClonerEffector/Source/ClonerEffector/Private/Cloner/Extensions/CEClonerConstraintExtension.cpp
--- a/Engine/Plugins/Experimental/ClonerEffector/Source/ClonerEffector/Private/Cloner/Extensions/CEClonerConstraintExtension.cpp
+++ b/Engine/Plugins/Experimental/ClonerEffector/Source/ClonerEffector/Private/Cloner/Extensions/CEClonerConstraintExtension.cpp
@@ -202,7 +202,8 @@ void UCEClonerConstraintExtension::OnExtensionParametersChanged(UCEClonerCompone
 	static const FNiagaraVariable ConstraintTextureCompareModeVar(FNiagaraTypeDefinition(StaticEnum<ECEClonerCompareMode>()), TEXT("ConstraintTextureCompareMode"));
 	ExposedParameters.SetParameterValue<int32>(static_cast<int32>(TextureCompareMode), ConstraintTextureCompareModeVar);
 
-	InComponent->SetFloatParameter(TEXT("ConstraintTextureThreshold"), TextureThreshold);
+	// To avoid resaving system asset and overwriting 5.6 version, use typo name in 5.5 
+	InComponent->SetFloatParameter(TEXT("ContraintTextureThreshold"), FMath::Max(0, TextureThreshold));
 }
 
 bool UCEClonerConstraintExtension::IsLayoutSupported(const UCEClonerLayoutBase* InLayout) const
