MotionDesign : ClonerEffector
- Add cancel effector to negate effects of other effectors based on effector shape
- Add intermediate class for effect extensions like forces, delay
- Allow filtering of effect extensions in effector mode
- Improving delay behavior with rotation and scale

#virtualized

diff --git a/Engine/Plugins/Experimental/ClonerEffector/Config/DefaultClonerEffector.ini b/Engine/Plugins/Experimental/ClonerEffector/Config/DefaultClonerEffector.ini
--- a/Engine/Plugins/Experimental/ClonerEffector/Config/DefaultClonerEffector.ini
+++ b/Engine/Plugins/Experimental/ClonerEffector/Config/DefaultClonerEffector.ini
@@ -22,6 +22,8 @@
 +ClassRedirects=(OldName="/Script/AvalancheEffectors.AvaEffectorComponent",NewName="/Script/ClonerEffector.CEEffectorComponent")
 +ClassRedirects=(OldName="/Script/AvalancheEffectors.AvaClonerSubsystem",NewName="/Script/ClonerEffector.CEClonerSubsystem")
 +ClassRedirects=(OldName="/Script/AvalancheEffectors.AvaEffectorSubsystem",NewName="/Script/ClonerEffector.CEEffectorSubsystem")
++ClassRedirects=(OldName="/Script/ClonerEffector.CEEffectorDelayExtension",NewName="/Script/ClonerEffector.CEEffectorDelayEffect")
++ClassRedirects=(OldName="/Script/ClonerEffector.CEEffectorForceExtension",NewName="/Script/ClonerEffector.CEEffectorForceEffect")
 
 +StructRedirects=(OldName="/Script/Avalanche.AvaClonerAttachmentItem",NewName="/Script/ClonerEffector.CEClonerAttachmentItem")
 +StructRedirects=(OldName="/Script/Avalanche.AvaClonerAttachmentTree",NewName="/Script/ClonerEffector.CEClonerAttachmentTree")
@@ -91,3 +93,7 @@
 +EnumRedirects=(OldName="/Script/Avalanche.EAvaClonerAttachmentStatus",NewName="/Script/ClonerEffector.ECEClonerAttachmentStatus")
 +EnumRedirects=(OldName="/Script/Avalanche.EAvaClonerTextureSampleChannel",NewName="/Script/ClonerEffector.ECEClonerTextureSampleChannel")
 +EnumRedirects=(OldName="/Script/Avalanche.EAvaClonerCompareMode",NewName="/Script/ClonerEffector.ECEClonerCompareMode")
+
++FunctionRedirects=(OldName="/Script/ClonerEffector.CEEffectorComponent.GetActiveExtensions",NewName="/Script/ClonerEffector.CEEffectorComponent.GetActiveEffects")
++FunctionRedirects=(OldName="/Script/ClonerEffector.CEEffectorLibrary.GetEffectorExtensionClasses",NewName="/Script/ClonerEffector.CEEffectorLibrary.GetEffectorEffectClasses")
+
diff --git a/Engine/Plugins/Experimental/ClonerEffector/Source/ClonerEffector/Private/Effector/CEEffectorActor.cpp b/Engine/Plugins/Experimental/ClonerEffector/Source/ClonerEffector/Private/Effector/CEEffectorActor.cpp
--- a/Engine/Plugins/Experimental/ClonerEffector/Source/ClonerEffector/Private/Effector/CEEffectorActor.cpp
+++ b/Engine/Plugins/Experimental/ClonerEffector/Source/ClonerEffector/Private/Effector/CEEffectorActor.cpp
@@ -4,7 +4,7 @@
 
 #include "Components/DynamicMeshComponent.h"
 #include "Effector/CEEffectorComponent.h"
-#include "Effector/Extensions/CEEffectorForceExtension.h"
+#include "Effector/Effects/CEEffectorForceEffect.h"
 #include "Effector/Modes/CEEffectorNoiseMode.h"
 #include "Effector/Modes/CEEffectorOffsetMode.h"
 #include "Effector/Modes/CEEffectorPushMode.h"
@@ -199,7 +199,7 @@ void ACEEffectorActor::MigrateDeprecatedProperties()
 			PushExtension->SetPushStrength(PushStrength);
 		}
 
-		if (UCEEffectorForceExtension* ForceExtension = EffectorComponent->FindOrAddExtension<UCEEffectorForceExtension>())
+		if (UCEEffectorForceEffect* ForceExtension = EffectorComponent->FindOrAddExtension<UCEEffectorForceEffect>())
 		{
 			ForceExtension->SetForcesEnabled(bAttractionForceEnabled || bGravityForceEnabled || bOrientationForceEnabled || bVortexForceEnabled || bCurlNoiseForceEnabled);
 
diff --git a/Engine/Plugins/Experimental/ClonerEffector/Source/ClonerEffector/Private/Effector/CEEffectorComponent.cpp b/Engine/Plugins/Experimental/ClonerEffector/Source/ClonerEffector/Private/Effector/CEEffectorComponent.cpp
--- a/Engine/Plugins/Experimental/ClonerEffector/Source/ClonerEffector/Private/Effector/CEEffectorComponent.cpp
+++ b/Engine/Plugins/Experimental/ClonerEffector/Source/ClonerEffector/Private/Effector/CEEffectorComponent.cpp
@@ -8,6 +8,7 @@
 #include "Components/DynamicMeshComponent.h"
 #include "Containers/Ticker.h"
 #include "Effector/CEEffectorExtensionBase.h"
+#include "Effector/Effects/CEEffectorEffectBase.h"
 #include "Effector/Modes/CEEffectorModeBase.h"
 #include "Effector/Types/CEEffectorTypeBase.h"
 #include "Engine/Texture2D.h"
@@ -210,9 +211,9 @@ void UCEEffectorComponent::OnClonerUnlinked(UCEClonerEffectorExtension* InCloner
 	ClonerExtensionsWeak.Remove(InCloner);
 }
 
-void UCEEffectorComponent::GetActiveExtensions(TArray<UCEEffectorExtensionBase*>& OutExtensions) const
+void UCEEffectorComponent::GetActiveEffects(TArray<UCEEffectorEffectBase*>& OutEffects) const
 {
-	OutExtensions = ActiveExtensions;
+	OutEffects = ActiveEffects;
 }
 
 UCEEffectorExtensionBase* UCEEffectorComponent::GetExtension(TSubclassOf<UCEEffectorExtensionBase> InExtensionClass) const
@@ -496,11 +497,11 @@ void UCEEffectorComponent::OnEffectorEnabled()
 		ActiveMode->ActivateExtension();
 	}
 
-	for (const TObjectPtr<UCEEffectorExtensionBase>& ActiveExtension : ActiveExtensions)
+	for (const TObjectPtr<UCEEffectorEffectBase>& ActiveEffect : ActiveEffects)
 	{
-		if (ActiveExtension)
+		if (ActiveEffect)
 		{
-			ActiveExtension->ActivateExtension();
+			ActiveEffect->ActivateExtension();
 		}
 	}
 
@@ -521,11 +522,11 @@ void UCEEffectorComponent::OnEffectorDisabled()
 		ActiveMode->DeactivateExtension();
 	}
 
-	for (const TObjectPtr<UCEEffectorExtensionBase>& ActiveExtension : ActiveExtensions)
+	for (const TObjectPtr<UCEEffectorEffectBase>& ActiveEffect : ActiveEffects)
 	{
-		if (ActiveExtension)
+		if (ActiveEffect)
 		{
-			ActiveExtension->DeactivateExtension();
+			ActiveEffect->DeactivateExtension();
 		}
 	}
 
@@ -565,12 +566,9 @@ void UCEEffectorComponent::OnEffectorOptionsChanged()
 	// Update Type
 	OnTypeNameChanged();
 
-	// Update Mode
+	// Update Mode & Effects
 	OnModeNameChanged();
 
-	// Update Extensions
-	OnExtensionsChanged();
-
 #if WITH_EDITOR
 	OnVisualizerOptionsChanged();
 #endif
@@ -642,47 +640,41 @@ void UCEEffectorComponent::OnModeNameChanged()
 			ActiveMode->UpdateExtensionParameters();
 		}
 	}
+
+	OnEffectsChanged();
 }
 
-void UCEEffectorComponent::OnExtensionsChanged()
+void UCEEffectorComponent::OnEffectsChanged()
 {
-	UCEEffectorSubsystem* Subsystem = UCEEffectorSubsystem::Get();
-
-	if (!Subsystem)
+	if (!ActiveMode)
 	{
 		return;
 	}
 
-	if (ActiveExtensions.IsEmpty())
+	TArray<TObjectPtr<UCEEffectorEffectBase>> RemoveEffects = ActiveEffects;
+	
+	for (const TSubclassOf<UCEEffectorEffectBase>& SupportedEffect : ActiveMode->GetSupportedEffects())
 	{
-		const TArray<FName> ModeNames = GetEffectorModeNames();
-		const TArray<FName> TypeNames = GetEffectorTypeNames();
-
-		for (const FName& ExtensionName : Subsystem->GetExtensionNames<UCEEffectorExtensionBase>())
+		if (UCEEffectorEffectBase* Effect = Cast<UCEEffectorEffectBase>(FindOrAddExtension(SupportedEffect)))
 		{
-			if (ModeNames.Contains(ExtensionName) || TypeNames.Contains(ExtensionName))
+			if (RemoveEffects.Remove(Effect) == 0)
 			{
-				continue;
-			}
+				Effect->ActivateExtension();
 
-			if (UCEEffectorExtensionBase* Extension = FindOrAddExtension(ExtensionName))
-			{
-				Extension->ActivateExtension();
-
-				ActiveExtensions.Add(Extension);
-
-				Extension->UpdateExtensionParameters();
+				ActiveEffects.Add(Effect);
 			}
+
+			Effect->UpdateExtensionParameters();
 		}
 	}
-	else
+
+	for (const TObjectPtr<UCEEffectorEffectBase>& ActiveEffect : RemoveEffects)
 	{
-		for (const TObjectPtr<UCEEffectorExtensionBase>& ActiveExtension : ActiveExtensions)
+		if (IsValid(ActiveEffect))
 		{
-			if (IsValid(ActiveExtension))
-			{
-				ActiveExtension->UpdateExtensionParameters();
-			}
+			ActiveEffect->DeactivateExtension();
+
+			ActiveEffects.Remove(ActiveEffect);
 		}
 	}
 }
diff --git a/Engine/Plugins/Experimental/ClonerEffector/Source/ClonerEffector/Private/Effector/Effects/CEEffectorDelayEffect.cpp b/Engine/Plugins/Experimental/ClonerEffector/Source/ClonerEffector/Private/Effector/Effects/CEEffectorDelayEffect.cpp
new file mode 100644
--- /dev/null
+++ b/Engine/Plugins/Experimental/ClonerEffector/Source/ClonerEffector/Private/Effector/Effects/CEEffectorDelayEffect.cpp
@@ -0,0 +1,101 @@
+// Copyright Epic Games, Inc. All Rights Reserved.
+
+#include "Effector/Effects/CEEffectorDelayEffect.h"
+
+#include "Effector/CEEffectorComponent.h"
+
+void UCEEffectorDelayEffect::SetDelayEnabled(bool bInEnabled)
+{
+	if (bDelayEnabled == bInEnabled)
+	{
+		return;
+	}
+
+	bDelayEnabled = bInEnabled;
+	UpdateExtensionParameters();
+}
+
+void UCEEffectorDelayEffect::SetDelayInDuration(float InDuration)
+{
+	InDuration = FMath::Max(InDuration, 0.f);
+	if (FMath::IsNearlyEqual(DelayInDuration, InDuration))
+	{
+		return;
+	}
+
+	DelayInDuration = InDuration;
+	UpdateExtensionParameters();
+}
+
+void UCEEffectorDelayEffect::SetDelayOutDuration(float InDuration)
+{
+	InDuration = FMath::Max(InDuration, 0.f);
+	if (FMath::IsNearlyEqual(DelayOutDuration, InDuration))
+	{
+		return;
+	}
+
+	DelayOutDuration = InDuration;
+	UpdateExtensionParameters();	
+}
+
+void UCEEffectorDelayEffect::SetDelaySpringFrequency(float InFrequency)
+{
+	InFrequency = FMath::Max(InFrequency, 1.f);
+	if (FMath::IsNearlyEqual(DelaySpringFrequency, InFrequency))
+	{
+		return;
+	}
+
+	DelaySpringFrequency = InFrequency;
+	UpdateExtensionParameters();
+}
+
+void UCEEffectorDelayEffect::SetDelaySpringFalloff(float InFalloff)
+{
+	InFalloff = FMath::Max(InFalloff, 1.f);
+	if (FMath::IsNearlyEqual(DelaySpringFalloff, InFalloff))
+	{
+		return;
+	}
+
+	DelaySpringFalloff = InFalloff;
+	UpdateExtensionParameters();
+}
+
+#if WITH_EDITOR
+const TCEPropertyChangeDispatcher<UCEEffectorDelayEffect> UCEEffectorDelayEffect::PropertyChangeDispatcher =
+{
+	{ GET_MEMBER_NAME_CHECKED(UCEEffectorDelayEffect, bDelayEnabled), &UCEEffectorDelayEffect::OnExtensionPropertyChanged },
+	{ GET_MEMBER_NAME_CHECKED(UCEEffectorDelayEffect, DelayInDuration), &UCEEffectorDelayEffect::OnExtensionPropertyChanged },
+	{ GET_MEMBER_NAME_CHECKED(UCEEffectorDelayEffect, DelaySpringFrequency), &UCEEffectorDelayEffect::OnExtensionPropertyChanged },
+	{ GET_MEMBER_NAME_CHECKED(UCEEffectorDelayEffect, DelaySpringFalloff), &UCEEffectorDelayEffect::OnExtensionPropertyChanged }
+};
+
+void UCEEffectorDelayEffect::PostEditChangeProperty(FPropertyChangedEvent& InPropertyChangedEvent)
+{
+	Super::PostEditChangeProperty(InPropertyChangedEvent);
+
+	PropertyChangeDispatcher.OnPropertyChanged(this, InPropertyChangedEvent);
+}
+#endif
+
+void UCEEffectorDelayEffect::UpdateEffectChannelData(FCEClonerEffectorChannelData& InChannelData, bool bInEnabled)
+{
+	Super::UpdateEffectChannelData(InChannelData, bInEnabled);
+
+	bInEnabled &= bDelayEnabled;
+	
+	if (bInEnabled)
+	{
+		InChannelData.DelayInDuration = DelayInDuration;
+		InChannelData.DelayOutDuration = DelayOutDuration;
+		InChannelData.DelaySpringFrequency = DelaySpringFrequency * DelayOutDuration;
+		InChannelData.DelaySpringFalloff = DelaySpringFalloff;
+	}
+	else
+	{
+		InChannelData.DelayInDuration = 0.f;
+		InChannelData.DelayOutDuration = 0.f;
+	}
+}
diff --git a/Engine/Plugins/Experimental/ClonerEffector/Source/ClonerEffector/Private/Effector/Effects/CEEffectorEffectBase.cpp b/Engine/Plugins/Experimental/ClonerEffector/Source/ClonerEffector/Private/Effector/Effects/CEEffectorEffectBase.cpp
new file mode 100644
--- /dev/null
+++ b/Engine/Plugins/Experimental/ClonerEffector/Source/ClonerEffector/Private/Effector/Effects/CEEffectorEffectBase.cpp
@@ -0,0 +1,22 @@
+// Copyright Epic Games, Inc. All Rights Reserved.
+
+#include "Effector/Effects/CEEffectorEffectBase.h"
+
+#include "Effector/CEEffectorComponent.h"
+
+void UCEEffectorEffectBase::OnExtensionParametersChanged(UCEEffectorComponent* InComponent)
+{
+	Super::OnExtensionParametersChanged(InComponent);
+	
+	UpdateEffectChannelData(InComponent->GetChannelData(), /** Enabled */true);
+}
+
+void UCEEffectorEffectBase::OnExtensionDeactivated()
+{
+	Super::OnExtensionDeactivated();
+
+	if (UCEEffectorComponent* EffectorComponent = GetEffectorComponent())
+	{
+		UpdateEffectChannelData(EffectorComponent->GetChannelData(), /** Enabled */false);
+	}
+}
diff --git a/Engine/Plugins/Experimental/ClonerEffector/Source/ClonerEffector/Private/Effector/Effects/CEEffectorForceEffect.cpp b/Engine/Plugins/Experimental/ClonerEffector/Source/ClonerEffector/Private/Effector/Effects/CEEffectorForceEffect.cpp
new file mode 100644
--- /dev/null
+++ b/Engine/Plugins/Experimental/ClonerEffector/Source/ClonerEffector/Private/Effector/Effects/CEEffectorForceEffect.cpp
@@ -0,0 +1,370 @@
+// Copyright Epic Games, Inc. All Rights Reserved.
+
+#include "Effector/Effects/CEEffectorForceEffect.h"
+
+#include "CEClonerEffectorShared.h"
+#include "Effector/CEEffectorComponent.h"
+
+void UCEEffectorForceEffect::SetForcesEnabled(bool bInForcesEnabled)
+{
+	if (bForcesEnabled == bInForcesEnabled)
+	{
+		return;
+	}
+
+	bForcesEnabled = bInForcesEnabled;
+	UpdateExtensionParameters(/** UpdateCloners */ true);
+}
+
+void UCEEffectorForceEffect::SetAttractionForceStrength(float InForceStrength)
+{
+	if (FMath::IsNearlyEqual(AttractionForceStrength, InForceStrength))
+	{
+		return;
+	}
+
+	AttractionForceStrength = InForceStrength;
+	UpdateExtensionParameters();
+}
+
+void UCEEffectorForceEffect::SetAttractionForceFalloff(float InForceFalloff)
+{
+	if (FMath::IsNearlyEqual(AttractionForceFalloff, InForceFalloff))
+	{
+		return;
+	}
+
+	AttractionForceFalloff = InForceFalloff;
+	UpdateExtensionParameters();
+}
+
+void UCEEffectorForceEffect::SetGravityForceEnabled(bool bInForceEnabled)
+{
+	if (bGravityForceEnabled == bInForceEnabled)
+	{
+		return;
+	}
+
+	bGravityForceEnabled = bInForceEnabled;
+	UpdateExtensionParameters(/** UpdateCloners */ true);
+}
+
+void UCEEffectorForceEffect::SetGravityForceAcceleration(const FVector& InAcceleration)
+{
+	if (GravityForceAcceleration.Equals(InAcceleration))
+	{
+		return;
+	}
+
+	GravityForceAcceleration = InAcceleration;
+	UpdateExtensionParameters();
+}
+
+void UCEEffectorForceEffect::SetDragForceEnabled(bool bInEnabled)
+{
+	if (bDragForceEnabled == bInEnabled)
+	{
+		return;
+	}
+
+	bDragForceEnabled = bInEnabled;
+	UpdateExtensionParameters(/** UpdateCloners */ true);
+}
+
+void UCEEffectorForceEffect::SetDragForceLinear(float InStrength)
+{
+	if (FMath::IsNearlyEqual(DragForceLinear, InStrength))
+	{
+		return;
+	}
+
+	DragForceLinear = InStrength;
+	UpdateExtensionParameters();
+}
+
+void UCEEffectorForceEffect::SetDragForceRotational(float InStrength)
+{
+	if (FMath::IsNearlyEqual(DragForceRotational, InStrength))
+	{
+		return;
+	}
+
+	DragForceRotational = InStrength;
+	UpdateExtensionParameters();
+}
+
+void UCEEffectorForceEffect::SetVectorNoiseForceEnabled(bool bInEnabled)
+{
+	if (bVectorNoiseForceEnabled == bInEnabled)
+	{
+		return;
+	}
+
+	bVectorNoiseForceEnabled = bInEnabled;
+	UpdateExtensionParameters(/** UpdateCloners */ true);
+}
+
+void UCEEffectorForceEffect::SetVectorNoiseForceAmount(float InAmount)
+{
+	if (FMath::IsNearlyEqual(VectorNoiseForceAmount, InAmount))
+	{
+		return;
+	}
+
+	VectorNoiseForceAmount = InAmount;
+	UpdateExtensionParameters();
+}
+
+void UCEEffectorForceEffect::SetOrientationForceEnabled(bool bInForceEnabled)
+{
+	if (bOrientationForceEnabled == bInForceEnabled)
+	{
+		return;
+	}
+
+	bOrientationForceEnabled = bInForceEnabled;
+	UpdateExtensionParameters(/** UpdateCloners */ true);
+}
+
+void UCEEffectorForceEffect::SetOrientationForceRate(float InForceOrientationRate)
+{
+	if (FMath::IsNearlyEqual(OrientationForceRate, InForceOrientationRate))
+	{
+		return;
+	}
+
+	OrientationForceRate = InForceOrientationRate;
+	UpdateExtensionParameters();
+}
+
+void UCEEffectorForceEffect::SetOrientationForceMin(const FVector& InForceOrientationMin)
+{
+	if (OrientationForceMin.Equals(InForceOrientationMin))
+	{
+		return;
+	}
+
+	OrientationForceMin = InForceOrientationMin;
+	UpdateExtensionParameters();
+}
+
+void UCEEffectorForceEffect::SetOrientationForceMax(const FVector& InForceOrientationMax)
+{
+	if (OrientationForceMax.Equals(InForceOrientationMax))
+	{
+		return;
+	}
+
+	OrientationForceMax = InForceOrientationMax;
+	UpdateExtensionParameters();
+}
+
+void UCEEffectorForceEffect::SetVortexForceEnabled(bool bInForceEnabled)
+{
+	if (bVortexForceEnabled == bInForceEnabled)
+	{
+		return;
+	}
+
+	bVortexForceEnabled = bInForceEnabled;
+	UpdateExtensionParameters(/** UpdateCloners */ true);
+}
+
+void UCEEffectorForceEffect::SetVortexForceAmount(float InForceVortexAmount)
+{
+	if (FMath::IsNearlyEqual(VortexForceAmount, InForceVortexAmount))
+	{
+		return;
+	}
+
+	VortexForceAmount = InForceVortexAmount;
+	UpdateExtensionParameters();
+}
+
+void UCEEffectorForceEffect::SetVortexForceAxis(const FVector& InForceVortexAxis)
+{
+	if (VortexForceAxis.Equals(InForceVortexAxis))
+	{
+		return;
+	}
+
+	VortexForceAxis = InForceVortexAxis;
+	UpdateExtensionParameters();
+}
+
+void UCEEffectorForceEffect::SetCurlNoiseForceEnabled(bool bInForceEnabled)
+{
+	if (bCurlNoiseForceEnabled == bInForceEnabled)
+	{
+		return;
+	}
+
+	bCurlNoiseForceEnabled = bInForceEnabled;
+	UpdateExtensionParameters(/** UpdateCloners */ true);
+}
+
+void UCEEffectorForceEffect::SetCurlNoiseForceStrength(float InForceCurlNoiseStrength)
+{
+	if (FMath::IsNearlyEqual(CurlNoiseForceStrength, InForceCurlNoiseStrength))
+	{
+		return;
+	}
+
+	CurlNoiseForceStrength = InForceCurlNoiseStrength;
+	UpdateExtensionParameters();
+}
+
+void UCEEffectorForceEffect::SetCurlNoiseForceFrequency(float InForceCurlNoiseFrequency)
+{
+	if (FMath::IsNearlyEqual(CurlNoiseForceFrequency, InForceCurlNoiseFrequency))
+	{
+		return;
+	}
+
+	CurlNoiseForceFrequency = InForceCurlNoiseFrequency;
+	UpdateExtensionParameters();
+}
+
+void UCEEffectorForceEffect::SetAttractionForceEnabled(bool bInForceEnabled)
+{
+	if (bAttractionForceEnabled == bInForceEnabled)
+	{
+		return;
+	}
+
+	bAttractionForceEnabled = bInForceEnabled;
+	UpdateExtensionParameters(/** UpdateCloners */ true);
+}
+
+void UCEEffectorForceEffect::UpdateEffectChannelData(FCEClonerEffectorChannelData& InChannelData, bool bInEnabled)
+{
+	Super::UpdateEffectChannelData(InChannelData, bInEnabled);
+
+	bInEnabled &= bForcesEnabled;
+
+	if (bInEnabled && bOrientationForceEnabled)
+	{
+		InChannelData.OrientationForceRate = OrientationForceRate;
+		InChannelData.OrientationForceMin = OrientationForceMin;
+		InChannelData.OrientationForceMax = OrientationForceMax;
+	}
+	else
+	{
+		InChannelData.OrientationForceRate = 0.f;
+		InChannelData.OrientationForceMin = FVector::ZeroVector;
+		InChannelData.OrientationForceMax = FVector::ZeroVector;
+	}
+
+	if (bInEnabled && bVortexForceEnabled)
+	{
+		InChannelData.VortexForceAmount = VortexForceAmount;
+		InChannelData.VortexForceAxis = VortexForceAxis;
+	}
+	else
+	{
+		InChannelData.VortexForceAmount = 0.f;
+		InChannelData.VortexForceAxis = FVector::ZeroVector;
+	}
+
+	if (bInEnabled && bCurlNoiseForceEnabled)
+	{
+		InChannelData.CurlNoiseForceStrength = CurlNoiseForceStrength;
+		InChannelData.CurlNoiseForceFrequency = CurlNoiseForceFrequency;
+	}
+	else
+	{
+		InChannelData.CurlNoiseForceStrength = 0.f;
+		InChannelData.CurlNoiseForceFrequency = 0.f;
+	}
+
+	if (bInEnabled && bAttractionForceEnabled)
+	{
+		InChannelData.AttractionForceStrength = AttractionForceStrength;
+		InChannelData.AttractionForceFalloff = AttractionForceFalloff;
+	}
+	else
+	{
+		InChannelData.AttractionForceStrength = 0.f;
+		InChannelData.AttractionForceFalloff = 0.f;
+	}
+
+	if (bInEnabled && bGravityForceEnabled)
+	{
+		InChannelData.GravityForceAcceleration = GravityForceAcceleration;
+	}
+	else
+	{
+		InChannelData.GravityForceAcceleration = FVector::ZeroVector;
+	}
+
+	if (bInEnabled && bDragForceEnabled)
+	{
+		InChannelData.DragForceLinear = DragForceLinear;
+		InChannelData.DragForceRotational = DragForceRotational;
+	}
+	else
+	{
+		InChannelData.DragForceLinear = 0.f;
+		InChannelData.DragForceRotational = 0.f;
+	}
+
+	if (bInEnabled && bVectorNoiseForceEnabled)
+	{
+		InChannelData.VectorNoiseForceAmount = VectorNoiseForceAmount;
+	}
+	else
+	{
+		InChannelData.VectorNoiseForceAmount = 0.f;
+	}
+}
+
+void UCEEffectorForceEffect::OnForceOptionsChanged()
+{
+	UpdateExtensionParameters(/** UpdateCloners */ true);
+}
+
+void UCEEffectorForceEffect::OnExtensionDeactivated()
+{
+	Super::OnExtensionDeactivated();
+
+	// Special case for forces to reset state
+	if (UCEEffectorComponent* EffectorComponent = GetEffectorComponent())
+	{
+		EffectorComponent->RequestClonerUpdate(/** Immediate */false);
+	}
+}
+
+#if WITH_EDITOR
+const TCEPropertyChangeDispatcher<UCEEffectorForceEffect> UCEEffectorForceEffect::PropertyChangeDispatcher =
+{
+	/** Force */
+	{ GET_MEMBER_NAME_CHECKED(UCEEffectorForceEffect, bForcesEnabled), &UCEEffectorForceEffect::OnForceOptionsChanged },
+	{ GET_MEMBER_NAME_CHECKED(UCEEffectorForceEffect, bOrientationForceEnabled), &UCEEffectorForceEffect::OnForceOptionsChanged },
+	{ GET_MEMBER_NAME_CHECKED(UCEEffectorForceEffect, OrientationForceRate), &UCEEffectorForceEffect::OnForceOptionsChanged },
+	{ GET_MEMBER_NAME_CHECKED(UCEEffectorForceEffect, OrientationForceMin), &UCEEffectorForceEffect::OnForceOptionsChanged },
+	{ GET_MEMBER_NAME_CHECKED(UCEEffectorForceEffect, OrientationForceMax), &UCEEffectorForceEffect::OnForceOptionsChanged },
+	{ GET_MEMBER_NAME_CHECKED(UCEEffectorForceEffect, bVortexForceEnabled), &UCEEffectorForceEffect::OnForceOptionsChanged },
+	{ GET_MEMBER_NAME_CHECKED(UCEEffectorForceEffect, VortexForceAmount), &UCEEffectorForceEffect::OnForceOptionsChanged },
+	{ GET_MEMBER_NAME_CHECKED(UCEEffectorForceEffect, VortexForceAxis), &UCEEffectorForceEffect::OnForceOptionsChanged },
+	{ GET_MEMBER_NAME_CHECKED(UCEEffectorForceEffect, bCurlNoiseForceEnabled), &UCEEffectorForceEffect::OnForceOptionsChanged },
+	{ GET_MEMBER_NAME_CHECKED(UCEEffectorForceEffect, CurlNoiseForceStrength), &UCEEffectorForceEffect::OnForceOptionsChanged },
+	{ GET_MEMBER_NAME_CHECKED(UCEEffectorForceEffect, CurlNoiseForceFrequency), &UCEEffectorForceEffect::OnForceOptionsChanged },
+	{ GET_MEMBER_NAME_CHECKED(UCEEffectorForceEffect, bAttractionForceEnabled), &UCEEffectorForceEffect::OnForceOptionsChanged },
+	{ GET_MEMBER_NAME_CHECKED(UCEEffectorForceEffect, AttractionForceStrength), &UCEEffectorForceEffect::OnForceOptionsChanged },
+	{ GET_MEMBER_NAME_CHECKED(UCEEffectorForceEffect, AttractionForceFalloff), &UCEEffectorForceEffect::OnForceOptionsChanged },
+	{ GET_MEMBER_NAME_CHECKED(UCEEffectorForceEffect, bGravityForceEnabled), &UCEEffectorForceEffect::OnForceOptionsChanged },
+	{ GET_MEMBER_NAME_CHECKED(UCEEffectorForceEffect, GravityForceAcceleration), &UCEEffectorForceEffect::OnForceOptionsChanged },
+	{ GET_MEMBER_NAME_CHECKED(UCEEffectorForceEffect, bDragForceEnabled), &UCEEffectorForceEffect::OnForceOptionsChanged },
+	{ GET_MEMBER_NAME_CHECKED(UCEEffectorForceEffect, DragForceLinear), &UCEEffectorForceEffect::OnForceOptionsChanged },
+	{ GET_MEMBER_NAME_CHECKED(UCEEffectorForceEffect, DragForceRotational), &UCEEffectorForceEffect::OnForceOptionsChanged },
+	{ GET_MEMBER_NAME_CHECKED(UCEEffectorForceEffect, bVectorNoiseForceEnabled), &UCEEffectorForceEffect::OnForceOptionsChanged },
+	{ GET_MEMBER_NAME_CHECKED(UCEEffectorForceEffect, VectorNoiseForceAmount), &UCEEffectorForceEffect::OnForceOptionsChanged },
+};
+
+void UCEEffectorForceEffect::PostEditChangeProperty(FPropertyChangedEvent& InPropertyChangedEvent)
+{
+	Super::PostEditChangeProperty(InPropertyChangedEvent);
+
+	PropertyChangeDispatcher.OnPropertyChanged(this, InPropertyChangedEvent);
+}
+#endif
diff --git a/Engine/Plugins/Experimental/ClonerEffector/Source/ClonerEffector/Private/Effector/Extensions/CEEffectorDelayExtension.cpp b/Engine/Plugins/Experimental/ClonerEffector/Source/ClonerEffector/Private/Effector/Extensions/CEEffectorDelayExtension.cpp
deleted file mode 100644
diff --git a/Engine/Plugins/Experimental/ClonerEffector/Source/ClonerEffector/Private/Effector/Extensions/CEEffectorForceExtension.cpp b/Engine/Plugins/Experimental/ClonerEffector/Source/ClonerEffector/Private/Effector/Extensions/CEEffectorForceExtension.cpp
deleted file mode 100644
diff --git a/Engine/Plugins/Experimental/ClonerEffector/Source/ClonerEffector/Private/Effector/Modes/CEEffectorCancelMode.cpp b/Engine/Plugins/Experimental/ClonerEffector/Source/ClonerEffector/Private/Effector/Modes/CEEffectorCancelMode.cpp
new file mode 100644
--- /dev/null
+++ b/Engine/Plugins/Experimental/ClonerEffector/Source/ClonerEffector/Private/Effector/Modes/CEEffectorCancelMode.cpp
@@ -0,0 +1,11 @@
+// Copyright Epic Games, Inc. All Rights Reserved.
+
+#include "Effector/Modes/CEEffectorCancelMode.h"
+
+#include "Effector/Effects/CEEffectorDelayEffect.h"
+#include "Effector/Effects/CEEffectorForceEffect.h"
+
+bool UCEEffectorCancelMode::IsEffectSupported(TSubclassOf<UCEEffectorEffectBase> InEffectClass) const
+{
+	return !(InEffectClass->IsChildOf<UCEEffectorDelayEffect>() || InEffectClass->IsChildOf<UCEEffectorForceEffect>());
+}
diff --git a/Engine/Plugins/Experimental/ClonerEffector/Source/ClonerEffector/Private/Effector/Modes/CEEffectorModeBase.cpp b/Engine/Plugins/Experimental/ClonerEffector/Source/ClonerEffector/Private/Effector/Modes/CEEffectorModeBase.cpp
--- a/Engine/Plugins/Experimental/ClonerEffector/Source/ClonerEffector/Private/Effector/Modes/CEEffectorModeBase.cpp
+++ b/Engine/Plugins/Experimental/ClonerEffector/Source/ClonerEffector/Private/Effector/Modes/CEEffectorModeBase.cpp
@@ -4,6 +4,46 @@
 
 #include "CEClonerEffectorShared.h"
 #include "Effector/CEEffectorComponent.h"
+#include "Effector/Effects/CEEffectorEffectBase.h"
+#include "Subsystems/CEEffectorSubsystem.h"
+
+TSet<TSubclassOf<UCEEffectorEffectBase>> UCEEffectorModeBase::GetSupportedEffects() const
+{
+	TSet<TSubclassOf<UCEEffectorEffectBase>> SupportedEffects;
+
+	if (const UCEEffectorSubsystem* EffectorSubsystem = UCEEffectorSubsystem::Get())
+	{
+		TSubclassOf<UCEEffectorModeBase> ModeClass(GetClass());
+		
+		for (const TSubclassOf<UCEEffectorExtensionBase>& ExtensionClass : EffectorSubsystem->GetExtensionClasses(UCEEffectorEffectBase::StaticClass()))
+		{
+			TSubclassOf<UCEEffectorEffectBase> EffectClass(ExtensionClass);
+			
+			const UCEEffectorEffectBase* Effect = EffectClass.GetDefaultObject();
+
+			if (!Effect)
+			{
+				continue;
+			}
+
+			// Does the mode supports this effect
+			if (!IsEffectSupported(EffectClass))
+			{
+				continue;
+			}
+
+			// Does the effect supports this mode
+			if (!Effect->IsModeSupported(ModeClass))
+			{
+				continue;
+			}
+
+			SupportedEffects.Add(EffectClass);
+		}
+	}
+
+	return SupportedEffects;
+}
 
 void UCEEffectorModeBase::OnExtensionParametersChanged(UCEEffectorComponent* InComponent)
 {
diff --git a/Engine/Plugins/Experimental/ClonerEffector/Source/ClonerEffector/Private/Utilities/CEEffectorLibrary.cpp b/Engine/Plugins/Experimental/ClonerEffector/Source/ClonerEffector/Private/Utilities/CEEffectorLibrary.cpp
--- a/Engine/Plugins/Experimental/ClonerEffector/Source/ClonerEffector/Private/Utilities/CEEffectorLibrary.cpp
+++ b/Engine/Plugins/Experimental/ClonerEffector/Source/ClonerEffector/Private/Utilities/CEEffectorLibrary.cpp
@@ -3,6 +3,7 @@
 #include "Utilities/CEEffectorLibrary.h"
 
 #include "Containers/Set.h"
+#include "Effector/Effects/CEEffectorEffectBase.h"
 #include "Effector/Modes/CEEffectorModeBase.h"
 #include "Effector/Types/CEEffectorTypeBase.h"
 #include "Subsystems/CEEffectorSubsystem.h"
@@ -33,22 +34,15 @@ void UCEEffectorLibrary::GetEffectorTypeClasses(TSet<TSubclassOf<UCEEffectorType
 	}
 }
 
-void UCEEffectorLibrary::GetEffectorExtensionClasses(TSet<TSubclassOf<UCEEffectorExtensionBase>>& OutExtensionClasses)
+void UCEEffectorLibrary::GetEffectorEffectClasses(TSet<TSubclassOf<UCEEffectorEffectBase>>& OutEffectClasses)
 {
-	OutExtensionClasses.Empty();
+	OutEffectClasses.Empty();
 
 	if (const UCEEffectorSubsystem* EffectorSubsystem = UCEEffectorSubsystem::Get())
 	{
-		for (TSubclassOf<UCEEffectorExtensionBase> ExtensionClass : EffectorSubsystem->GetExtensionClasses<UCEEffectorExtensionBase>())
+		for (TSubclassOf<UCEEffectorExtensionBase> EffectClass : EffectorSubsystem->GetExtensionClasses<UCEEffectorEffectBase>())
 		{
-			if (!ExtensionClass.Get()
-				|| ExtensionClass->IsChildOf<UCEEffectorTypeBase>()
-				|| ExtensionClass->IsChildOf<UCEEffectorModeBase>())
-			{
-				continue;
-			}
-
-			OutExtensionClasses.Add(ExtensionClass);
+			OutEffectClasses.Add(EffectClass.Get());
 		}
 	}
 }
diff --git a/Engine/Plugins/Experimental/ClonerEffector/Source/ClonerEffector/Public/CEClonerEffectorShared.h b/Engine/Plugins/Experimental/ClonerEffector/Source/ClonerEffector/Public/CEClonerEffectorShared.h
--- a/Engine/Plugins/Experimental/ClonerEffector/Source/ClonerEffector/Public/CEClonerEffectorShared.h
+++ b/Engine/Plugins/Experimental/ClonerEffector/Source/ClonerEffector/Public/CEClonerEffectorShared.h
@@ -163,7 +163,9 @@ enum class ECEClonerEffectorMode : uint8
 	/** Accumulate transform on clones based on their index */
 	Step,
 	/** Culls clones based on a zone */
-	Cull
+	Cull,
+	/** Cancels/reverses all effects */
+	Cancel
 };
 
 /** Enumerates the effector push direction available */
diff --git a/Engine/Plugins/Experimental/ClonerEffector/Source/ClonerEffector/Public/Cloner/Layouts/CEClonerLayoutBase.h b/Engine/Plugins/Experimental/ClonerEffector/Source/ClonerEffector/Public/Cloner/Layouts/CEClonerLayoutBase.h
--- a/Engine/Plugins/Experimental/ClonerEffector/Source/ClonerEffector/Public/Cloner/Layouts/CEClonerLayoutBase.h
+++ b/Engine/Plugins/Experimental/ClonerEffector/Source/ClonerEffector/Public/Cloner/Layouts/CEClonerLayoutBase.h
@@ -18,10 +18,12 @@ enum class ECEClonerLayoutAssetVersion : uint8
 	PreVersioning = 0,
 	/** Niagara System is now cached to improve loading time */
 	CachedVersion,
-	/** Cull clones feature added */
+	/** Cull clones added */
 	CullEffector,
-	/** Delay clones feature added */
+	/** Delay clones added */
 	DelayEffector,
+	/** Cancel clones effect added */
+	CancelEffector,
 	LatestVersionPlusOne,
 	LatestVersion = LatestVersionPlusOne - 1
 };
diff --git a/Engine/Plugins/Experimental/ClonerEffector/Source/ClonerEffector/Public/Effector/CEEffectorActor.h b/Engine/Plugins/Experimental/ClonerEffector/Source/ClonerEffector/Public/Effector/CEEffectorActor.h
--- a/Engine/Plugins/Experimental/ClonerEffector/Source/ClonerEffector/Public/Effector/CEEffectorActor.h
+++ b/Engine/Plugins/Experimental/ClonerEffector/Source/ClonerEffector/Public/Effector/CEEffectorActor.h
@@ -13,7 +13,7 @@ class UDynamicMeshComponent;
 UCLASS(
 	MinimalAPI
 	, BlueprintType
-	, HideCategories=(Rendering,Replication,Collision,HLOD,Physics,Networking,Input,Actor,Cooking,LevelInstance,DataLayers)
+	, HideCategories=(Rendering,Replication,Collision,HLOD,Physics,Networking,Input,Actor,Cooking,LevelInstance,DataLayers,WorldPartition)
 	, DisplayName = "Motion Design Effector Actor")
 class ACEEffectorActor : public AActor
 {
diff --git a/Engine/Plugins/Experimental/ClonerEffector/Source/ClonerEffector/Public/Effector/CEEffectorComponent.h b/Engine/Plugins/Experimental/ClonerEffector/Source/ClonerEffector/Public/Effector/CEEffectorComponent.h
--- a/Engine/Plugins/Experimental/ClonerEffector/Source/ClonerEffector/Public/Effector/CEEffectorComponent.h
+++ b/Engine/Plugins/Experimental/ClonerEffector/Source/ClonerEffector/Public/Effector/CEEffectorComponent.h
@@ -9,6 +9,7 @@
 #include "CEEffectorComponent.generated.h"
 
 class UCEClonerEffectorExtension;
+class UCEEffectorEffectBase;
 class UCEEffectorExtensionBase;
 class UCEEffectorModeBase;
 class UCEEffectorTypeBase;
@@ -163,13 +164,13 @@ public:
 		return ActiveMode;
 	}
 
-	TConstArrayView<TObjectPtr<UCEEffectorExtensionBase>> GetActiveExtensions() const
+	TConstArrayView<TObjectPtr<UCEEffectorEffectBase>> GetActiveEffects() const
 	{
-		return ActiveExtensions;
+		return ActiveEffects;
 	}
 
 	UFUNCTION(BlueprintPure, Category="Effector")
-	void GetActiveExtensions(TArray<UCEEffectorExtensionBase*>& OutExtensions) const;
+	void GetActiveEffects(TArray<UCEEffectorEffectBase*>& OutEffects) const;
 
 	UFUNCTION(BlueprintCallable, Category="Effector")
 	UCEEffectorExtensionBase* GetExtension(TSubclassOf<UCEEffectorExtensionBase> InExtensionClass) const;
@@ -227,9 +228,9 @@ protected:
 	UPROPERTY(Transient, DuplicateTransient, NonTransactional)
 	TObjectPtr<UCEEffectorModeBase> ActiveMode;
 
-	/** Active Extensions on this effector */
+	/** Active Effects on this effector */
 	UPROPERTY(Transient, DuplicateTransient, NonTransactional)
-	TArray<TObjectPtr<UCEEffectorExtensionBase>> ActiveExtensions;
+	TArray<TObjectPtr<UCEEffectorEffectBase>> ActiveEffects;
 
 #if WITH_EDITORONLY_DATA
 	/** Visibility of the components visualizer */
@@ -257,7 +258,7 @@ private:
 	void OnEffectorOptionsChanged();
 	void OnTypeNameChanged();
 	void OnModeNameChanged();
-	void OnExtensionsChanged();
+	void OnEffectsChanged();
 
 #if WITH_EDITOR
 	void OnVisualizerOptionsChanged();
diff --git a/Engine/Plugins/Experimental/ClonerEffector/Source/ClonerEffector/Public/Effector/Effects/CEEffectorDelayEffect.h b/Engine/Plugins/Experimental/ClonerEffector/Source/ClonerEffector/Public/Effector/Effects/CEEffectorDelayEffect.h
new file mode 100644
--- /dev/null
+++ b/Engine/Plugins/Experimental/ClonerEffector/Source/ClonerEffector/Public/Effector/Effects/CEEffectorDelayEffect.h
@@ -0,0 +1,105 @@
+// Copyright Epic Games, Inc. All Rights Reserved.
+
+#pragma once
+
+#include "CEPropertyChangeDispatcher.h"
+#include "Effector/Effects/CEEffectorEffectBase.h"
+#include "CEEffectorDelayEffect.generated.h"
+
+class UCEEffectorComponent;
+
+/** Effector delay extension to affect clones after a period of time with a spring effect */
+UCLASS(MinimalAPI, BlueprintType, Within=CEEffectorComponent, meta=(Section="Delay", Priority=3))
+class UCEEffectorDelayEffect : public UCEEffectorEffectBase
+{
+	GENERATED_BODY()
+
+public:
+	UCEEffectorDelayEffect()
+		: UCEEffectorEffectBase(
+			TEXT("Delay")
+		)
+	{}
+
+	UFUNCTION(BlueprintCallable, Category="Effector")
+	CLONEREFFECTOR_API void SetDelayEnabled(bool bInEnabled);
+
+	UFUNCTION(BlueprintPure, Category="Effector")
+	bool GetDelayEnabled() const
+	{
+		return bDelayEnabled;
+	}
+
+	UFUNCTION(BlueprintCallable, Category="Effector")
+	CLONEREFFECTOR_API void SetDelayInDuration(float InDuration);
+
+	UFUNCTION(BlueprintPure, Category="Effector")
+	float GetDelayInDuration() const
+	{
+		return DelayInDuration;
+	}
+
+	UFUNCTION(BlueprintCallable, Category="Effector")
+	CLONEREFFECTOR_API void SetDelayOutDuration(float InDuration);
+
+	UFUNCTION(BlueprintPure, Category="Effector")
+	float GetDelayOutDuration() const
+	{
+		return DelayOutDuration;
+	}
+
+	UFUNCTION(BlueprintCallable, Category="Effector")
+	CLONEREFFECTOR_API void SetDelaySpringFrequency(float InFrequency);
+
+	UFUNCTION(BlueprintPure, Category="Effector")
+	float GetDelaySpringFrequency() const
+	{
+		return DelaySpringFrequency;
+	}
+
+	UFUNCTION(BlueprintCallable, Category="Effector")
+	CLONEREFFECTOR_API void SetDelaySpringFalloff(float InFalloff);
+
+	UFUNCTION(BlueprintPure, Category="Effector")
+	float GetDelaySpringFalloff() const
+	{
+		return DelaySpringFalloff;
+	}
+
+protected:
+	//~ Begin UObject
+#if WITH_EDITOR
+	virtual void PostEditChangeProperty(FPropertyChangedEvent& InPropertyChangedEvent) override;
+#endif
+	//~ End UObject
+
+	//~ Begin UCEEffectorEffectBase
+	virtual void UpdateEffectChannelData(FCEClonerEffectorChannelData& InChannelData, bool bInEnabled) override;
+	//~ End UCEEffectorEffectBase
+
+	/** Enable time delay for the mode effect */
+	UPROPERTY(EditInstanceOnly, Setter="SetDelayEnabled", Getter="GetDelayEnabled", DisplayName="Enabled", Category="Delay")
+	bool bDelayEnabled = false;
+	
+	/** Time delay for the effect to reach its peak */
+	UPROPERTY(EditInstanceOnly, Setter, Getter, DisplayName="Offset Duration", Category="Delay", meta=(Units=Seconds, ClampMin="0", EditCondition="bDelayEnabled", EditConditionHides))
+	float DelayInDuration = 1.f;
+
+	/** Time delay for the effect to go back to its rest state */
+	UPROPERTY(EditInstanceOnly, Setter, Getter, DisplayName="Trail Duration", Category="Delay", meta=(Units=Seconds, ClampMin="0", EditCondition="bDelayEnabled", EditConditionHides))
+	float DelayOutDuration = 1.f;
+
+	/** Frequency of the delayed spring effect */
+	UPROPERTY(EditInstanceOnly, Setter, Getter, DisplayName="Spring Frequency", Category="Delay", meta=(ClampMin="1", EditCondition="bDelayEnabled", EditConditionHides))
+	float DelaySpringFrequency = 3.f;
+
+	/** Damping or decay of the delayed spring effect */
+	UPROPERTY(EditInstanceOnly, Setter, Getter, DisplayName="Spring Falloff", Category="Delay", meta=(ClampMin="1", EditCondition="bDelayEnabled", EditConditionHides))
+	float DelaySpringFalloff = 1.f;
+
+private:
+#if WITH_EDITOR
+	/** Used for PECP */
+	static const TCEPropertyChangeDispatcher<UCEEffectorDelayEffect> PropertyChangeDispatcher;
+#endif
+};
\ No newline at end of file
diff --git a/Engine/Plugins/Experimental/ClonerEffector/Source/ClonerEffector/Public/Effector/Effects/CEEffectorEffectBase.h b/Engine/Plugins/Experimental/ClonerEffector/Source/ClonerEffector/Public/Effector/Effects/CEEffectorEffectBase.h
new file mode 100644
--- /dev/null
+++ b/Engine/Plugins/Experimental/ClonerEffector/Source/ClonerEffector/Public/Effector/Effects/CEEffectorEffectBase.h
@@ -0,0 +1,40 @@
+// Copyright Epic Games, Inc. All Rights Reserved.
+
+#pragma once
+
+#include "Effector/CEEffectorExtensionBase.h"
+#include "CEEffectorEffectBase.generated.h"
+
+class UCEEffectorModeBase;
+
+/** Represents an effect for an effector to affect clones in a specific way, works with modes and types */
+UCLASS(MinimalAPI, Abstract, BlueprintType, Within=CEEffectorComponent, meta=(Section="Effect", Priority=3))
+class UCEEffectorEffectBase : public UCEEffectorExtensionBase
+{
+	GENERATED_BODY()
+
+public:
+	UCEEffectorEffectBase()
+		: UCEEffectorExtensionBase(NAME_None)
+	{}
+
+	UCEEffectorEffectBase(FName InModeName)
+		: UCEEffectorExtensionBase(
+			InModeName
+		)
+	{}
+
+	/** Filter supported mode */
+	virtual bool IsModeSupported(TSubclassOf<UCEEffectorModeBase> InModeClass) const
+	{
+		return true;
+	}
+
+protected:
+	//~ Begin UCEEffectorExtensionBase
+	virtual void OnExtensionParametersChanged(UCEEffectorComponent* InComponent) override;
+	virtual void OnExtensionDeactivated() override;
+	//~ End UCEEffectorExtensionBase
+
+	virtual void UpdateEffectChannelData(FCEClonerEffectorChannelData& InChannelData, bool bInEnabled) {}
+};
\ No newline at end of file
diff --git a/Engine/Plugins/Experimental/ClonerEffector/Source/ClonerEffector/Public/Effector/Effects/CEEffectorForceEffect.h b/Engine/Plugins/Experimental/ClonerEffector/Source/ClonerEffector/Public/Effector/Effects/CEEffectorForceEffect.h
new file mode 100644
--- /dev/null
+++ b/Engine/Plugins/Experimental/ClonerEffector/Source/ClonerEffector/Public/Effector/Effects/CEEffectorForceEffect.h
@@ -0,0 +1,306 @@
+// Copyright Epic Games, Inc. All Rights Reserved.
+
+#pragma once
+
+#include "CEPropertyChangeDispatcher.h"
+#include "Effector/Effects/CEEffectorEffectBase.h"
+#include "CEEffectorForceEffect.generated.h"
+
+class UCEEffectorComponent;
+
+/** Effector supported forces to affect clones */
+UCLASS(MinimalAPI, BlueprintType, Within=CEEffectorComponent, meta=(Section="Forces", Priority=3))
+class UCEEffectorForceEffect : public UCEEffectorEffectBase
+{
+	GENERATED_BODY()
+
+public:
+	UCEEffectorForceEffect()
+		: UCEEffectorEffectBase(
+			TEXT("Forces")
+		)
+	{}
+
+	UFUNCTION(BlueprintCallable, Category="Effector")
+	CLONEREFFECTOR_API void SetForcesEnabled(bool bInForcesEnabled);
+
+	UFUNCTION(BlueprintPure, Category="Effector")
+	bool GetForcesEnabled() const
+	{
+		return bForcesEnabled;
+	}
+
+	UFUNCTION(BlueprintCallable, Category="Effector")
+	CLONEREFFECTOR_API void SetOrientationForceEnabled(bool bInForceEnabled);
+
+	UFUNCTION(BlueprintPure, Category="Effector")
+	bool GetOrientationForceEnabled() const
+	{
+		return bOrientationForceEnabled;
+	}
+
+	UFUNCTION(BlueprintCallable, Category="Effector")
+	CLONEREFFECTOR_API void SetOrientationForceRate(float InForceOrientationRate);
+
+	UFUNCTION(BlueprintPure, Category="Effector")
+	float GetOrientationForceRate() const
+	{
+		return OrientationForceRate;
+	}
+
+	UFUNCTION(BlueprintCallable, Category="Effector")
+	CLONEREFFECTOR_API void SetOrientationForceMin(const FVector& InForceOrientationMin);
+
+	UFUNCTION(BlueprintPure, Category="Effector")
+	const FVector& GetOrientationForceMin() const
+	{
+		return OrientationForceMin;
+	}
+
+	UFUNCTION(BlueprintCallable, Category="Effector")
+	CLONEREFFECTOR_API void SetOrientationForceMax(const FVector& InForceOrientationMax);
+
+	UFUNCTION(BlueprintPure, Category="Effector")
+	const FVector& GetOrientationForceMax() const
+	{
+		return OrientationForceMax;
+	}
+
+	UFUNCTION(BlueprintCallable, Category="Effector")
+	CLONEREFFECTOR_API void SetVortexForceEnabled(bool bInForceEnabled);
+
+	UFUNCTION(BlueprintPure, Category="Effector")
+	bool GetVortexForceEnabled() const
+	{
+		return bVortexForceEnabled;
+	}
+
+	UFUNCTION(BlueprintCallable, Category="Effector")
+	CLONEREFFECTOR_API void SetVortexForceAmount(float InForceVortexAmount);
+
+	UFUNCTION(BlueprintPure, Category="Effector")
+	float GetVortexForceAmount() const
+	{
+		return VortexForceAmount;
+	}
+
+	UFUNCTION(BlueprintCallable, Category="Effector")
+	CLONEREFFECTOR_API void SetVortexForceAxis(const FVector& InForceVortexAxis);
+
+	UFUNCTION(BlueprintPure, Category="Effector")
+	const FVector& GetVortexForceAxis() const
+	{
+		return VortexForceAxis;
+	}
+
+	UFUNCTION(BlueprintCallable, Category="Effector")
+	CLONEREFFECTOR_API void SetCurlNoiseForceEnabled(bool bInForceEnabled);
+
+	UFUNCTION(BlueprintPure, Category="Effector")
+	bool GetCurlNoiseForceEnabled() const
+	{
+		return bCurlNoiseForceEnabled;
+	}
+
+	UFUNCTION(BlueprintCallable, Category="Effector")
+	CLONEREFFECTOR_API void SetCurlNoiseForceStrength(float InForceCurlNoiseStrength);
+
+	UFUNCTION(BlueprintPure, Category="Effector")
+	float GetCurlNoiseForceStrength() const
+	{
+		return CurlNoiseForceStrength;
+	}
+
+	UFUNCTION(BlueprintCallable, Category="Effector")
+	CLONEREFFECTOR_API void SetCurlNoiseForceFrequency(float InForceCurlNoiseFrequency);
+
+	UFUNCTION(BlueprintPure, Category="Effector")
+	float GetCurlNoiseForceFrequency() const
+	{
+		return CurlNoiseForceFrequency;
+	}
+
+	UFUNCTION(BlueprintCallable, Category="Effector")
+	CLONEREFFECTOR_API void SetAttractionForceEnabled(bool bInForceEnabled);
+
+	UFUNCTION(BlueprintPure, Category="Effector")
+	bool GetAttractionForceEnabled() const
+	{
+		return bAttractionForceEnabled;
+	}
+
+	UFUNCTION(BlueprintCallable, Category="Effector")
+	CLONEREFFECTOR_API void SetAttractionForceStrength(float InForceStrength);
+
+	UFUNCTION(BlueprintPure, Category="Effector")
+	float GetAttractionForceStrength() const
+	{
+		return AttractionForceStrength;
+	}
+
+	UFUNCTION(BlueprintCallable, Category="Effector")
+	CLONEREFFECTOR_API void SetAttractionForceFalloff(float InForceFalloff);
+
+	UFUNCTION(BlueprintPure, Category="Effector")
+	float GetAttractionForceFalloff() const
+	{
+		return AttractionForceFalloff;
+	}
+
+	UFUNCTION(BlueprintCallable, Category="Effector")
+	CLONEREFFECTOR_API void SetGravityForceEnabled(bool bInForceEnabled);
+
+	UFUNCTION(BlueprintPure, Category="Effector")
+	bool GetGravityForceEnabled() const
+	{
+		return bGravityForceEnabled;
+	}
+
+	UFUNCTION(BlueprintCallable, Category="Effector")
+	CLONEREFFECTOR_API void SetGravityForceAcceleration(const FVector& InAcceleration);
+
+	UFUNCTION(BlueprintPure, Category="Effector")
+	FVector GetGravityForceAcceleration() const
+	{
+		return GravityForceAcceleration;
+	}
+
+	UFUNCTION(BlueprintCallable, Category="Effector")
+	CLONEREFFECTOR_API void SetDragForceEnabled(bool bInEnabled);
+
+	UFUNCTION(BlueprintPure, Category="Effector")
+	bool GetDragForceEnabled() const
+	{
+		return bDragForceEnabled;
+	}
+
+	UFUNCTION(BlueprintCallable, Category="Effector")
+	CLONEREFFECTOR_API void SetDragForceLinear(float InStrength);
+
+	UFUNCTION(BlueprintPure, Category="Effector")
+	float GetDragForceLinear() const
+	{
+		return DragForceLinear;
+	}
+
+	UFUNCTION(BlueprintCallable, Category="Effector")
+	CLONEREFFECTOR_API void SetDragForceRotational(float InStrength);
+
+	UFUNCTION(BlueprintPure, Category="Effector")
+	float GetDragForceRotational() const
+	{
+		return DragForceRotational;
+	}
+
+	UFUNCTION(BlueprintCallable, Category="Effector")
+	CLONEREFFECTOR_API void SetVectorNoiseForceEnabled(bool bInEnabled);
+
+	UFUNCTION(BlueprintPure, Category="Effector")
+	bool GetVectorNoiseForceEnabled() const
+	{
+		return bVectorNoiseForceEnabled;
+	}
+
+	UFUNCTION(BlueprintCallable, Category="Effector")
+	CLONEREFFECTOR_API void SetVectorNoiseForceAmount(float InAmount);
+
+	UFUNCTION(BlueprintPure, Category="Effector")
+	float GetVectorNoiseForceAmount() const
+	{
+		return VectorNoiseForceAmount;
+	}
+
+protected:
+	//~ Begin UObject
+#if WITH_EDITOR
+	virtual void PostEditChangeProperty(FPropertyChangedEvent& InPropertyChangedEvent) override;
+#endif
+	//~ End UObject
+
+	//~ Begin UCEEffectorExtensionBase
+	virtual void OnExtensionDeactivated() override;
+	//~ End UCEEffectorExtensionBase
+	
+	//~ Begin UCEEffectorEffectBase
+	virtual void UpdateEffectChannelData(FCEClonerEffectorChannelData& InChannelData, bool bInEnabled) override;
+	//~ End UCEEffectorEffectBase
+
+	void OnForceOptionsChanged();
+
+	/** Forces global state */
+	UPROPERTY(EditInstanceOnly, Setter="SetForcesEnabled", Getter="GetForcesEnabled", DisplayName="Enabled", Category="Force")
+	bool bForcesEnabled = false;
+
+	/** Enable orientation force to allow each clone instance to rotate around its pivot */
+	UPROPERTY(EditInstanceOnly, Setter="SetOrientationForceEnabled", Getter="GetOrientationForceEnabled", Category="Force", meta=(EditCondition="bForcesEnabled", EditConditionHides))
+	bool bOrientationForceEnabled = false;
+
+	UPROPERTY(EditInstanceOnly, Setter, Getter, Category="Force", meta=(Delta="0.0001", ClampMin="0", EditCondition="bForcesEnabled && bOrientationForceEnabled", EditConditionHides))
+	float OrientationForceRate = 1.f;
+
+	UPROPERTY(EditInstanceOnly, Setter, Getter, Category="Force", meta=(EditCondition="bForcesEnabled && bOrientationForceEnabled", EditConditionHides))
+	FVector OrientationForceMin = FVector(-0.1f);
+
+	UPROPERTY(EditInstanceOnly, Setter, Getter, Category="Force", meta=(EditCondition="bForcesEnabled && bOrientationForceEnabled", EditConditionHides))
+	FVector OrientationForceMax = FVector(0.1f);
+
+	/** Enable vortex force to allow each clone instance to rotate around a specific axis on the cloner pivot */
+	UPROPERTY(EditInstanceOnly, Setter="SetVortexForceEnabled", Getter="GetVortexForceEnabled", Category="Force", meta=(EditCondition="bForcesEnabled", EditConditionHides))
+	bool bVortexForceEnabled = false;
+
+	UPROPERTY(EditInstanceOnly, Setter, Getter, Category="Force", meta=(EditCondition="bForcesEnabled && bVortexForceEnabled", EditConditionHides))
+	float VortexForceAmount = 100.f;
+
+	UPROPERTY(EditInstanceOnly, Setter, Getter, Category="Force", meta=(EditCondition="bForcesEnabled && bVortexForceEnabled", EditConditionHides))
+	FVector VortexForceAxis = FVector::ZAxisVector;
+
+	/** Enable curl noise force to allow each clone instance to add random location variation */
+	UPROPERTY(EditInstanceOnly, Setter="SetCurlNoiseForceEnabled", Getter="GetCurlNoiseForceEnabled", Category="Force", meta=(EditCondition="bForcesEnabled", EditConditionHides))
+	bool bCurlNoiseForceEnabled = false;
+
+	UPROPERTY(EditInstanceOnly, Setter, Getter, Category="Force", meta=(EditCondition="bForcesEnabled && bCurlNoiseForceEnabled", EditConditionHides))
+	float CurlNoiseForceStrength = 100.f;
+
+	UPROPERTY(EditInstanceOnly, Setter, Getter, Category="Force", meta=(EditCondition="bForcesEnabled && bCurlNoiseForceEnabled", EditConditionHides))
+	float CurlNoiseForceFrequency = 10.f;
+
+	/** Enable attraction force to allow each clone instances to gravitate toward a location */
+	UPROPERTY(EditInstanceOnly, Setter="SetAttractionForceEnabled", Getter="GetAttractionForceEnabled", Category="Force", meta=(EditCondition="bForcesEnabled", EditConditionHides))
+	bool bAttractionForceEnabled = false;
+
+	UPROPERTY(EditInstanceOnly, Setter, Getter, Category="Force", meta=(EditCondition="bForcesEnabled && bAttractionForceEnabled", EditConditionHides))
+	float AttractionForceStrength = 100.f;
+
+	UPROPERTY(EditInstanceOnly, Setter, Getter, Category="Force", meta=(EditCondition="bForcesEnabled && bAttractionForceEnabled", EditConditionHides))
+	float AttractionForceFalloff = 0.1f;
+
+	/** Enable gravity force to pull particles based on an acceleration vector */
+	UPROPERTY(EditInstanceOnly, Setter="SetGravityForceEnabled", Getter="GetGravityForceEnabled", Category="Force", meta=(EditCondition="bForcesEnabled", EditConditionHides))
+	bool bGravityForceEnabled = false;
+
+	UPROPERTY(EditInstanceOnly, Setter, Getter, Category="Force", meta=(EditCondition="bForcesEnabled && bGravityForceEnabled", EditConditionHides))
+	FVector GravityForceAcceleration = FVector(0, 0, -980.f);
+
+	/** Enable drag force to decrease particles velocity */
+	UPROPERTY(EditInstanceOnly, Setter="SetDragForceEnabled", Getter="GetDragForceEnabled", Category="Force", meta=(EditCondition="bForcesEnabled", EditConditionHides))
+	bool bDragForceEnabled = false;
+
+	UPROPERTY(EditInstanceOnly, Setter, Getter, Category="Force", meta=(EditCondition="bForcesEnabled && bDragForceEnabled", EditConditionHides))
+	float DragForceLinear = 0.f;
+
+	UPROPERTY(EditInstanceOnly, Setter, Getter, Category="Force", meta=(EditCondition="bForcesEnabled && bDragForceEnabled", EditConditionHides))
+	float DragForceRotational = 0.f;
+
+	/** Enable vector random noise force to add variation to clones behavior */
+	UPROPERTY(EditInstanceOnly, Setter="SetVectorNoiseForceEnabled", Getter="GetVectorNoiseForceEnabled", Category="Force", meta=(EditCondition="bForcesEnabled", EditConditionHides))
+	bool bVectorNoiseForceEnabled = false;
+
+	UPROPERTY(EditInstanceOnly, Setter, Getter, Category="Force", meta=(EditCondition="bForcesEnabled && bVectorNoiseForceEnabled", EditConditionHides))
+	float VectorNoiseForceAmount = 100.f;
+
+private:
+#if WITH_EDITOR
+	/** Used for PECP */
+	static const TCEPropertyChangeDispatcher<UCEEffectorForceEffect> PropertyChangeDispatcher;
+#endif
+};
\ No newline at end of file
diff --git a/Engine/Plugins/Experimental/ClonerEffector/Source/ClonerEffector/Public/Effector/Extensions/CEEffectorDelayExtension.h b/Engine/Plugins/Experimental/ClonerEffector/Source/ClonerEffector/Public/Effector/Extensions/CEEffectorDelayExtension.h
deleted file mode 100644
diff --git a/Engine/Plugins/Experimental/ClonerEffector/Source/ClonerEffector/Public/Effector/Extensions/CEEffectorForceExtension.h b/Engine/Plugins/Experimental/ClonerEffector/Source/ClonerEffector/Public/Effector/Extensions/CEEffectorForceExtension.h
deleted file mode 100644
diff --git a/Engine/Plugins/Experimental/ClonerEffector/Source/ClonerEffector/Public/Effector/Modes/CEEffectorCancelMode.h b/Engine/Plugins/Experimental/ClonerEffector/Source/ClonerEffector/Public/Effector/Modes/CEEffectorCancelMode.h
new file mode 100644
--- /dev/null
+++ b/Engine/Plugins/Experimental/ClonerEffector/Source/ClonerEffector/Public/Effector/Modes/CEEffectorCancelMode.h
@@ -0,0 +1,26 @@
+// Copyright Epic Games, Inc. All Rights Reserved.
+
+#pragma once
+
+#include "CEClonerEffectorShared.h"
+#include "CEEffectorModeBase.h"
+#include "CEEffectorCancelMode.generated.h"
+
+class UCEEffectorComponent;
+
+/** Negates all effects applied on clones */
+UCLASS(MinimalAPI, BlueprintType, Within=CEEffectorComponent)
+class UCEEffectorCancelMode : public UCEEffectorModeBase
+{
+	GENERATED_BODY()
+
+public:
+	UCEEffectorCancelMode()
+		: UCEEffectorModeBase(TEXT("Cancel"), static_cast<int32>(ECEClonerEffectorMode::Cancel))
+	{}
+
+protected:
+	//~ Begin UCEEffectorModeBase
+	virtual bool IsEffectSupported(TSubclassOf<UCEEffectorEffectBase> InEffectClass) const override;
+	//~ End UCEEffectorModeBase
+};
diff --git a/Engine/Plugins/Experimental/ClonerEffector/Source/ClonerEffector/Public/Effector/Modes/CEEffectorModeBase.h b/Engine/Plugins/Experimental/ClonerEffector/Source/ClonerEffector/Public/Effector/Modes/CEEffectorModeBase.h
--- a/Engine/Plugins/Experimental/ClonerEffector/Source/ClonerEffector/Public/Effector/Modes/CEEffectorModeBase.h
+++ b/Engine/Plugins/Experimental/ClonerEffector/Source/ClonerEffector/Public/Effector/Modes/CEEffectorModeBase.h
@@ -6,6 +6,7 @@
 #include "CEEffectorModeBase.generated.h"
 
 class UCEEffectorComponent;
+class UCEEffectorEffectBase;
 
 /** Represents a behavior for an effector to affect clones in a specific way */
 UCLASS(MinimalAPI, Abstract, BlueprintType, Within=CEEffectorComponent, meta=(Section="Mode", Priority=2))
@@ -29,12 +30,20 @@ public:
 	{
 		return ModeIdentifier;
 	}
+	
+	TSet<TSubclassOf<UCEEffectorEffectBase>> GetSupportedEffects() const;
 
 protected:
 	//~ Begin UCEEffectorExtensionBase
 	virtual void OnExtensionParametersChanged(UCEEffectorComponent* InComponent) override;
 	//~ End UCEEffectorExtensionBase
 
+	/** Override if a specific effect is not supported by this mode */
+	virtual bool IsEffectSupported(TSubclassOf<UCEEffectorEffectBase> InEffectClass) const
+	{
+		return true;
+	}
+
 private:
 	/** Unique identifier to pass it to niagara */
 	UPROPERTY(Transient)
diff --git a/Engine/Plugins/Experimental/ClonerEffector/Source/ClonerEffector/Public/Utilities/CEEffectorLibrary.h b/Engine/Plugins/Experimental/ClonerEffector/Source/ClonerEffector/Public/Utilities/CEEffectorLibrary.h
--- a/Engine/Plugins/Experimental/ClonerEffector/Source/ClonerEffector/Public/Utilities/CEEffectorLibrary.h
+++ b/Engine/Plugins/Experimental/ClonerEffector/Source/ClonerEffector/Public/Utilities/CEEffectorLibrary.h
@@ -7,6 +7,7 @@
 #include "Templates/SubclassOf.h"
 #include "CEEffectorLibrary.generated.h"
 
+class UCEEffectorEffectBase;
 class UCEEffectorExtensionBase;
 class UCEEffectorTypeBase;
 class UCEEffectorModeBase;
@@ -33,9 +34,9 @@ public:
 	static CLONEREFFECTOR_API void GetEffectorTypeClasses(TSet<TSubclassOf<UCEEffectorTypeBase>>& OutTypeClasses);
 
 	/**
-	 * Retrieves all extension classes available for an effector
-	 * @param OutExtensionClasses [Out] Available extension classes
+	 * Retrieves all effect classes available for an effector
+	 * @param OutEffectClasses [Out] Available effect classes
 	 */
 	UFUNCTION(BlueprintCallable, Category="Motion Design|Effector|Utility")
-	static CLONEREFFECTOR_API void GetEffectorExtensionClasses(TSet<TSubclassOf<UCEEffectorExtensionBase>>& OutExtensionClasses);
+	static CLONEREFFECTOR_API void GetEffectorEffectClasses(TSet<TSubclassOf<UCEEffectorEffectBase>>& OutEffectClasses);
 };
\ No newline at end of file
diff --git a/Engine/Plugins/Experimental/ClonerEffector/Source/ClonerEffectorEditor/Private/Effector/Customizations/CEEditorEffectorComponentDetailCustomization.cpp b/Engine/Plugins/Experimental/ClonerEffector/Source/ClonerEffectorEditor/Private/Effector/Customizations/CEEditorEffectorComponentDetailCustomization.cpp
--- a/Engine/Plugins/Experimental/ClonerEffector/Source/ClonerEffectorEditor/Private/Effector/Customizations/CEEditorEffectorComponentDetailCustomization.cpp
+++ b/Engine/Plugins/Experimental/ClonerEffector/Source/ClonerEffectorEditor/Private/Effector/Customizations/CEEditorEffectorComponentDetailCustomization.cpp
@@ -6,6 +6,7 @@
 #include "DetailWidgetRow.h"
 #include "Effector/CEEffectorComponent.h"
 #include "Effector/CEEffectorExtensionBase.h"
+#include "Effector/Effects/CEEffectorEffectBase.h"
 #include "Effector/Modes/CEEffectorModeBase.h"
 #include "Effector/Types/CEEffectorTypeBase.h"
 #include "IPropertyUtilities.h"
@@ -80,12 +81,12 @@ void FCEEditorEffectorComponentDetailCustomization::CustomizeDetails(IDetailLayo
 			CategoryData.Objects.Add(ActiveMode);
 		}
 
-		for (const TObjectPtr<UCEEffectorExtensionBase>& ActiveExtension : EffectorComponent->GetActiveExtensions())
+		for (const TObjectPtr<UCEEffectorEffectBase>& ActiveEffect : EffectorComponent->GetActiveEffects())
 		{
-			FDetailsCategoryData& CategoryData = CategoryToData.FindOrAdd(ActiveExtension->GetExtensionName());
-			CategoryData.SectionName = ActiveExtension->GetExtensionSection().SectionName;
-			CategoryData.SortOrder = StartSortOrder + ExtensionSortOrderOffset + ActiveExtension->GetExtensionSection().SectionOrder;
-			CategoryData.Objects.Add(ActiveExtension);
+			FDetailsCategoryData& CategoryData = CategoryToData.FindOrAdd(ActiveEffect->GetExtensionName());
+			CategoryData.SectionName = ActiveEffect->GetExtensionSection().SectionName;
+			CategoryData.SortOrder = StartSortOrder + ExtensionSortOrderOffset + ActiveEffect->GetExtensionSection().SectionOrder;
+			CategoryData.Objects.Add(ActiveEffect);
 		}
 	}
 
