MotionDesign : ClonerEffector
- Refactor effector noise mode to be procedural mode with multiple patterns
- Add support for directional wave and circular wave patterns
- Rework curl noise pattern to have similar frequency, pan capability and amplitude as the other patterns

#virtualized

diff --git a/Engine/Plugins/Experimental/ClonerEffector/Config/DefaultClonerEffector.ini b/Engine/Plugins/Experimental/ClonerEffector/Config/DefaultClonerEffector.ini
--- a/Engine/Plugins/Experimental/ClonerEffector/Config/DefaultClonerEffector.ini
+++ b/Engine/Plugins/Experimental/ClonerEffector/Config/DefaultClonerEffector.ini
@@ -24,6 +24,7 @@
 +ClassRedirects=(OldName="/Script/AvalancheEffectors.AvaEffectorSubsystem",NewName="/Script/ClonerEffector.CEEffectorSubsystem")
 +ClassRedirects=(OldName="/Script/ClonerEffector.CEEffectorDelayExtension",NewName="/Script/ClonerEffector.CEEffectorDelayEffect")
 +ClassRedirects=(OldName="/Script/ClonerEffector.CEEffectorForceExtension",NewName="/Script/ClonerEffector.CEEffectorForceEffect")
++ClassRedirects=(OldName="/Script/ClonerEffector.CEEffectorNoiseMode",NewName="/Script/ClonerEffector.CEEffectorProceduralMode")
 
 +StructRedirects=(OldName="/Script/Avalanche.AvaClonerAttachmentItem",NewName="/Script/ClonerEffector.CEClonerAttachmentItem")
 +StructRedirects=(OldName="/Script/Avalanche.AvaClonerAttachmentTree",NewName="/Script/ClonerEffector.CEClonerAttachmentTree")
@@ -93,7 +94,7 @@
 +EnumRedirects=(OldName="/Script/Avalanche.EAvaClonerAttachmentStatus",NewName="/Script/ClonerEffector.ECEClonerAttachmentStatus")
 +EnumRedirects=(OldName="/Script/Avalanche.EAvaClonerTextureSampleChannel",NewName="/Script/ClonerEffector.ECEClonerTextureSampleChannel")
 +EnumRedirects=(OldName="/Script/Avalanche.EAvaClonerCompareMode",NewName="/Script/ClonerEffector.ECEClonerCompareMode")
++EnumRedirects=(OldName="/Script/ClonerEffector.ECEClonerEffectorMode",ValueChanges=(("NoiseField","Procedural")))
 
 +FunctionRedirects=(OldName="/Script/ClonerEffector.CEEffectorComponent.GetActiveExtensions",NewName="/Script/ClonerEffector.CEEffectorComponent.GetActiveEffects")
 +FunctionRedirects=(OldName="/Script/ClonerEffector.CEEffectorLibrary.GetEffectorExtensionClasses",NewName="/Script/ClonerEffector.CEEffectorLibrary.GetEffectorEffectClasses")
-
diff --git a/Engine/Plugins/Experimental/ClonerEffector/Source/ClonerEffector/Private/CEClonerEffectorShared.cpp b/Engine/Plugins/Experimental/ClonerEffector/Source/ClonerEffector/Private/CEClonerEffectorShared.cpp
--- a/Engine/Plugins/Experimental/ClonerEffector/Source/ClonerEffector/Private/CEClonerEffectorShared.cpp
+++ b/Engine/Plugins/Experimental/ClonerEffector/Source/ClonerEffector/Private/CEClonerEffectorShared.cpp
@@ -31,21 +31,28 @@ void FCEClonerEffectorChannelData::Write(UNiagaraDataChannelWriter* InWriter) co
 	}
 
 	/** General */
-	InWriter->WriteInt(EasingName, Identifier, static_cast<int32>(Easing));
-	InWriter->WriteInt(ModeName, Identifier, static_cast<int32>(Mode));
-	InWriter->WriteInt(TypeName, Identifier, static_cast<int32>(Type));
 	InWriter->WriteFloat(MagnitudeName, Identifier, Magnitude);
+	InWriter->WritePosition(LocationName, Identifier, Location);
+	InWriter->WriteQuat(RotationName, Identifier, Rotation);
+	InWriter->WriteVector(ScaleName, Identifier, Scale);
+	InWriter->WriteLinearColor(ColorName, Identifier, Color);
+
+	/** Type/Shape */
+	InWriter->WriteInt(TypeName, Identifier, static_cast<int32>(Type));
+	InWriter->WriteInt(EasingName, Identifier, static_cast<int32>(Easing));
 	InWriter->WriteVector(InnerExtentName, Identifier, InnerExtent);
 	InWriter->WriteVector(OuterExtentName, Identifier, OuterExtent);
+	
+	/** Mode */
+	InWriter->WriteInt(ModeName, Identifier, static_cast<int32>(Mode));
 	InWriter->WriteVector(LocationDeltaName, Identifier, LocationDelta);
 	InWriter->WriteVector(RotationDeltaName, Identifier, RotationDelta);
 	InWriter->WriteVector(ScaleDeltaName, Identifier, ScaleDelta);
-	InWriter->WritePosition(LocationName, Identifier, Location);
-	InWriter->WriteQuat(RotationName, Identifier, Rotation);
-	InWriter->WriteVector(ScaleName, Identifier, Scale);
 	InWriter->WriteFloat(FrequencyName, Identifier, Frequency);
 	InWriter->WriteVector(PanName, Identifier, Pan);
-	InWriter->WriteLinearColor(ColorName, Identifier, Color);
+	InWriter->WriteInt(PatternName, Identifier, static_cast<int32>(Pattern));
+	
+	/** Delay */
 	InWriter->WriteFloat(DelayInDurationName, Identifier, DelayInDuration);
 	InWriter->WriteFloat(DelayOutDurationName, Identifier, DelayOutDuration);
 	InWriter->WriteFloat(DelaySpringFrequencyName, Identifier, DelaySpringFrequency);
@@ -73,24 +80,32 @@ void FCEClonerEffectorChannelData::Read(const UNiagaraDataChannelReader* InReade
 	{
 		return;
 	}
-
-	/** Effector */
+	
 	bool bIsValid;
-	Easing = static_cast<ECEClonerEasing>(InReader->ReadInt(EasingName, Identifier, bIsValid));
-	Mode = static_cast<ECEClonerEffectorMode>(InReader->ReadInt(ModeName, Identifier, bIsValid));
-	Type = static_cast<ECEClonerEffectorType>(InReader->ReadInt(TypeName, Identifier, bIsValid));
+
+	/** General */
 	Magnitude = InReader->ReadFloat(MagnitudeName, Identifier, bIsValid);
+	Location = InReader->ReadPosition(LocationName, Identifier, bIsValid);
+	Rotation = InReader->ReadQuat(RotationName, Identifier, bIsValid);
+	Scale = InReader->ReadVector(ScaleName, Identifier, bIsValid);
+	Color = InReader->ReadLinearColor(ColorName, Identifier, bIsValid);
+
+	/** Type/Shape */
+	Type = static_cast<ECEClonerEffectorType>(InReader->ReadInt(TypeName, Identifier, bIsValid));
+	Easing = static_cast<ECEClonerEasing>(InReader->ReadInt(EasingName, Identifier, bIsValid));
 	InnerExtent = InReader->ReadVector(InnerExtentName, Identifier, bIsValid);
 	OuterExtent = InReader->ReadVector(OuterExtentName, Identifier, bIsValid);
+	
+	/** Mode */
+	Mode = static_cast<ECEClonerEffectorMode>(InReader->ReadInt(ModeName, Identifier, bIsValid));
 	LocationDelta = InReader->ReadVector(LocationDeltaName, Identifier, bIsValid);
 	RotationDelta = InReader->ReadVector(RotationDeltaName, Identifier, bIsValid);
 	ScaleDelta = InReader->ReadVector(ScaleDeltaName, Identifier, bIsValid);
-	Location = InReader->ReadPosition(LocationName, Identifier, bIsValid);
-	Rotation = InReader->ReadQuat(RotationName, Identifier, bIsValid);
-	Scale = InReader->ReadVector(ScaleName, Identifier, bIsValid);
 	Frequency = InReader->ReadFloat(FrequencyName, Identifier, bIsValid);
 	Pan = InReader->ReadVector(PanName, Identifier, bIsValid);
-	Color = InReader->ReadLinearColor(ColorName, Identifier, bIsValid);
+	Pattern = static_cast<ECEClonerEffectorProceduralPattern>(InReader->ReadInt(PatternName, Identifier, bIsValid));
+	
+	/** Delay */
 	DelayInDuration = InReader->ReadFloat(DelayInDurationName, Identifier, bIsValid);
 	DelayOutDuration = InReader->ReadFloat(DelayOutDurationName, Identifier, bIsValid);
 	DelaySpringFrequency = InReader->ReadFloat(DelaySpringFrequencyName, Identifier, bIsValid);
diff --git a/Engine/Plugins/Experimental/ClonerEffector/Source/ClonerEffector/Private/Effector/CEEffectorActor.cpp b/Engine/Plugins/Experimental/ClonerEffector/Source/ClonerEffector/Private/Effector/CEEffectorActor.cpp
--- a/Engine/Plugins/Experimental/ClonerEffector/Source/ClonerEffector/Private/Effector/CEEffectorActor.cpp
+++ b/Engine/Plugins/Experimental/ClonerEffector/Source/ClonerEffector/Private/Effector/CEEffectorActor.cpp
@@ -5,8 +5,8 @@
 #include "Components/DynamicMeshComponent.h"
 #include "Effector/CEEffectorComponent.h"
 #include "Effector/Effects/CEEffectorForceEffect.h"
-#include "Effector/Modes/CEEffectorNoiseMode.h"
 #include "Effector/Modes/CEEffectorOffsetMode.h"
+#include "Effector/Modes/CEEffectorProceduralMode.h"
 #include "Effector/Modes/CEEffectorPushMode.h"
 #include "Effector/Modes/CEEffectorTargetMode.h"
 #include "Effector/Types/CEEffectorBoxType.h"
@@ -184,7 +184,7 @@ void ACEEffectorActor::MigrateDeprecatedProperties()
 			TargetExtension->SetTargetActorWeak(TargetActorWeak);
 		}
 
-		if (UCEEffectorNoiseMode* NoiseExtension = EffectorComponent->FindOrAddExtension<UCEEffectorNoiseMode>())
+		if (UCEEffectorProceduralMode* NoiseExtension = EffectorComponent->FindOrAddExtension<UCEEffectorProceduralMode>())
 		{
 			NoiseExtension->SetLocationStrength(LocationStrength);
 			NoiseExtension->SetRotationStrength(RotationStrength);
diff --git a/Engine/Plugins/Experimental/ClonerEffector/Source/ClonerEffector/Private/Effector/CEEffectorComponent.cpp b/Engine/Plugins/Experimental/ClonerEffector/Source/ClonerEffector/Private/Effector/CEEffectorComponent.cpp
--- a/Engine/Plugins/Experimental/ClonerEffector/Source/ClonerEffector/Private/Effector/CEEffectorComponent.cpp
+++ b/Engine/Plugins/Experimental/ClonerEffector/Source/ClonerEffector/Private/Effector/CEEffectorComponent.cpp
@@ -582,6 +582,18 @@ void UCEEffectorComponent::OnTypeNameChanged()
 	if (!TypeNames.Contains(TypeName) && !TypeNames.IsEmpty())
 	{
 		TypeName = TypeNames[0];
+		
+		for (const FName& ExtensionName : TypeNames)
+		{
+			if (const UCEEffectorExtensionBase* ExtensionType = GetExtension(ExtensionName))
+			{
+				if (ExtensionType->RedirectExtensionName(TypeName))
+				{
+					TypeName = ExtensionType->GetExtensionName();
+					break;
+				}
+			}
+		}
 	}
 
 	if (UCEEffectorTypeBase* Type = Cast<UCEEffectorTypeBase>(FindOrAddExtension(TypeName)))
@@ -616,6 +628,18 @@ void UCEEffectorComponent::OnModeNameChanged()
 	if (!ModeNames.Contains(ModeName) && !ModeNames.IsEmpty())
 	{
 		ModeName = ModeNames[0];
+		
+		for (const FName& ExtensionName : ModeNames)
+		{
+			if (const UCEEffectorExtensionBase* ExtensionMode = GetExtension(ExtensionName))
+			{
+				if (ExtensionMode->RedirectExtensionName(ModeName))
+				{
+					ModeName = ExtensionMode->GetExtensionName();
+					break;
+				}
+			}
+		}
 	}
 
 	if (UCEEffectorModeBase* Mode = Cast<UCEEffectorModeBase>(FindOrAddExtension(ModeName)))
diff --git a/Engine/Plugins/Experimental/ClonerEffector/Source/ClonerEffector/Private/Effector/Modes/CEEffectorNoiseMode.cpp b/Engine/Plugins/Experimental/ClonerEffector/Source/ClonerEffector/Private/Effector/Modes/CEEffectorNoiseMode.cpp
deleted file mode 100644
diff --git a/Engine/Plugins/Experimental/ClonerEffector/Source/ClonerEffector/Private/Effector/Modes/CEEffectorProceduralMode.cpp b/Engine/Plugins/Experimental/ClonerEffector/Source/ClonerEffector/Private/Effector/Modes/CEEffectorProceduralMode.cpp
new file mode 100644
--- /dev/null
+++ b/Engine/Plugins/Experimental/ClonerEffector/Source/ClonerEffector/Private/Effector/Modes/CEEffectorProceduralMode.cpp
@@ -0,0 +1,128 @@
+// Copyright Epic Games, Inc. All Rights Reserved.
+
+#include "Effector/Modes/CEEffectorProceduralMode.h"
+
+#include "CEClonerEffectorShared.h"
+#include "Effector/CEEffectorComponent.h"
+#include "NiagaraTypeRegistry.h"
+
+void UCEEffectorProceduralMode::SetPattern(ECEClonerEffectorProceduralPattern InPattern)
+{
+	if (InPattern == Pattern)
+	{
+		return;
+	}
+
+	Pattern = InPattern;
+	UpdateExtensionParameters();
+}
+
+void UCEEffectorProceduralMode::SetLocationStrength(const FVector& InStrength)
+{
+	if (LocationStrength.Equals(InStrength))
+	{
+		return;
+	}
+
+	LocationStrength = InStrength;
+	UpdateExtensionParameters();
+}
+
+void UCEEffectorProceduralMode::SetRotationStrength(const FRotator& InStrength)
+{
+	if (RotationStrength.Equals(InStrength))
+	{
+		return;
+	}
+
+	RotationStrength = InStrength;
+	UpdateExtensionParameters();
+}
+
+void UCEEffectorProceduralMode::SetScaleStrength(const FVector& InStrength)
+{
+	if (ScaleStrength.Equals(InStrength))
+	{
+		return;
+	}
+
+	ScaleStrength = InStrength;
+	UpdateExtensionParameters();
+}
+
+void UCEEffectorProceduralMode::SetPan(const FVector& InPan)
+{
+	if (Pan.Equals(InPan))
+	{
+		return;
+	}
+
+	Pan = InPan;
+	UpdateExtensionParameters();
+}
+
+void UCEEffectorProceduralMode::SetFrequency(float InFrequency)
+{
+	InFrequency = FMath::Max(InFrequency, 0);
+
+	if (FMath::IsNearlyEqual(Frequency, InFrequency))
+	{
+		return;
+	}
+
+	Frequency = InFrequency;
+	UpdateExtensionParameters();
+}
+
+void UCEEffectorProceduralMode::OnExtensionParametersChanged(UCEEffectorComponent* InComponent)
+{
+	Super::OnExtensionParametersChanged(InComponent);
+
+	FCEClonerEffectorChannelData& ChannelData = InComponent->GetChannelData();
+	ChannelData.LocationDelta = LocationStrength;
+	ChannelData.RotationDelta = FVector(RotationStrength.Yaw, RotationStrength.Pitch, RotationStrength.Roll);
+	ChannelData.ScaleDelta = ScaleStrength;
+	ChannelData.Frequency = Frequency;
+	ChannelData.Pan = Pan;
+	ChannelData.Pattern = Pattern;
+}
+
+bool UCEEffectorProceduralMode::RedirectExtensionName(FName InOldExtensionName) const
+{
+	return InOldExtensionName.IsEqual(TEXT("Noise")) || Super::RedirectExtensionName(InOldExtensionName);
+}
+
+void UCEEffectorProceduralMode::PostInitProperties()
+{
+	Super::PostInitProperties();
+
+	if (HasAnyFlags(RF_ClassDefaultObject))
+	{
+		// Register new type def for niagara
+
+		constexpr ENiagaraTypeRegistryFlags MeshFlags =
+			ENiagaraTypeRegistryFlags::AllowAnyVariable |
+			ENiagaraTypeRegistryFlags::AllowParameter;
+
+		FNiagaraTypeRegistry::Register(FNiagaraTypeDefinition(StaticEnum<ECEClonerEffectorProceduralPattern>()), MeshFlags);
+	}
+}
+
+#if WITH_EDITOR
+const TCEPropertyChangeDispatcher<UCEEffectorProceduralMode> UCEEffectorProceduralMode::PropertyChangeDispatcher =
+{
+	{ GET_MEMBER_NAME_CHECKED(UCEEffectorProceduralMode, LocationStrength), &UCEEffectorProceduralMode::OnExtensionPropertyChanged },
+	{ GET_MEMBER_NAME_CHECKED(UCEEffectorProceduralMode, RotationStrength), &UCEEffectorProceduralMode::OnExtensionPropertyChanged },
+	{ GET_MEMBER_NAME_CHECKED(UCEEffectorProceduralMode, ScaleStrength), &UCEEffectorProceduralMode::OnExtensionPropertyChanged },
+	{ GET_MEMBER_NAME_CHECKED(UCEEffectorProceduralMode, Pan), &UCEEffectorProceduralMode::OnExtensionPropertyChanged },
+	{ GET_MEMBER_NAME_CHECKED(UCEEffectorProceduralMode, Frequency), &UCEEffectorProceduralMode::OnExtensionPropertyChanged },
+	{ GET_MEMBER_NAME_CHECKED(UCEEffectorProceduralMode, Pattern), &UCEEffectorProceduralMode::OnExtensionPropertyChanged },
+};
+
+void UCEEffectorProceduralMode::PostEditChangeProperty(FPropertyChangedEvent& InPropertyChangedEvent)
+{
+	Super::PostEditChangeProperty(InPropertyChangedEvent);
+
+	PropertyChangeDispatcher.OnPropertyChanged(this, InPropertyChangedEvent);
+}
+#endif
diff --git a/Engine/Plugins/Experimental/ClonerEffector/Source/ClonerEffector/Private/Subsystems/CEEffectorSubsystem.cpp b/Engine/Plugins/Experimental/ClonerEffector/Source/ClonerEffector/Private/Subsystems/CEEffectorSubsystem.cpp
--- a/Engine/Plugins/Experimental/ClonerEffector/Source/ClonerEffector/Private/Subsystems/CEEffectorSubsystem.cpp
+++ b/Engine/Plugins/Experimental/ClonerEffector/Source/ClonerEffector/Private/Subsystems/CEEffectorSubsystem.cpp
@@ -8,8 +8,8 @@
 #include "NiagaraDataChannelPublic.h"
 #include "Effector/CEEffectorComponent.h"
 #include "Effector/CEEffectorExtensionBase.h"
-#include "Effector/Modes/CEEffectorNoiseMode.h"
 #include "Effector/Modes/CEEffectorOffsetMode.h"
+#include "Effector/Modes/CEEffectorProceduralMode.h"
 #include "Effector/Modes/CEEffectorPushMode.h"
 #include "Effector/Modes/CEEffectorTargetMode.h"
 #include "Effector/Types/CEEffectorBoxType.h"
@@ -61,7 +61,7 @@ void UCEEffectorSubsystem::Initialize(FSubsystemCollectionBase& Collection)
 	// Register modes
 	RegisterExtensionClass(UCEEffectorOffsetMode::StaticClass());
 	RegisterExtensionClass(UCEEffectorTargetMode::StaticClass());
-	RegisterExtensionClass(UCEEffectorNoiseMode::StaticClass());
+	RegisterExtensionClass(UCEEffectorProceduralMode::StaticClass());
 	RegisterExtensionClass(UCEEffectorPushMode::StaticClass());
 
 	ScanForRegistrableClasses();
diff --git a/Engine/Plugins/Experimental/ClonerEffector/Source/ClonerEffector/Public/CEClonerEffectorShared.h b/Engine/Plugins/Experimental/ClonerEffector/Source/ClonerEffector/Public/CEClonerEffectorShared.h
--- a/Engine/Plugins/Experimental/ClonerEffector/Source/ClonerEffector/Public/CEClonerEffectorShared.h
+++ b/Engine/Plugins/Experimental/ClonerEffector/Source/ClonerEffector/Public/CEClonerEffectorShared.h
@@ -156,8 +156,8 @@ enum class ECEClonerEffectorMode : uint8
 	Default,
 	/** Rotates clones towards a target actor */
 	Target,
-	/** Randomly applies noise across the field zone */
-	NoiseField,
+	/** Procedural pattern applied across the field zone */
+	Procedural,
 	/** Pushes clones apart based on a strength and direction */
 	Push,
 	/** Accumulate transform on clones based on their index */
@@ -168,6 +168,18 @@ enum class ECEClonerEffectorMode : uint8
 	Cancel
 };
 
+/** Enumerates the effector procedural pattern available */
+UENUM(BlueprintType)
+enum class ECEClonerEffectorProceduralPattern : uint8
+{
+	/** Jacobian Simplex Noise */
+	CurlNoise,
+	/** Wave moving towards specific direction */
+	DirectionalWave,
+	/** Ripple wave movement */
+	CircularWave
+};
+
 /** Enumerates the effector push direction available */
 UENUM(BlueprintType)
 enum class ECEClonerEffectorPushDirection : uint8
@@ -376,21 +388,28 @@ struct FCEClonerEffectorChannelData
 	friend class UCEEffectorSubsystem;
 
 	/** General */
-	static constexpr const TCHAR* EasingName = TEXT("Easing");
-	static constexpr const TCHAR* ModeName = TEXT("Mode");
-	static constexpr const TCHAR* TypeName = TEXT("Type");
 	static constexpr const TCHAR* MagnitudeName = TEXT("Magnitude");
+	static constexpr const TCHAR* LocationName = TEXT("Location");
+	static constexpr const TCHAR* RotationName = TEXT("Rotation");
+	static constexpr const TCHAR* ScaleName = TEXT("Scale");
+	static constexpr const TCHAR* ColorName = TEXT("Color");
+
+	/** Type/Shape */
+	static constexpr const TCHAR* TypeName = TEXT("Type");
+	static constexpr const TCHAR* EasingName = TEXT("Easing");
 	static constexpr const TCHAR* InnerExtentName = TEXT("InnerExtent");
 	static constexpr const TCHAR* OuterExtentName = TEXT("OuterExtent");
+
+	/** Mode */
+	static constexpr const TCHAR* ModeName = TEXT("Mode");
 	static constexpr const TCHAR* LocationDeltaName = TEXT("LocationDelta");
 	static constexpr const TCHAR* RotationDeltaName = TEXT("RotationDelta");
 	static constexpr const TCHAR* ScaleDeltaName = TEXT("ScaleDelta");
-	static constexpr const TCHAR* LocationName = TEXT("Location");
-	static constexpr const TCHAR* RotationName = TEXT("Rotation");
-	static constexpr const TCHAR* ScaleName = TEXT("Scale");
 	static constexpr const TCHAR* FrequencyName = TEXT("Frequency");
 	static constexpr const TCHAR* PanName = TEXT("Pan");
-	static constexpr const TCHAR* ColorName = TEXT("Color");
+	static constexpr const TCHAR* PatternName = TEXT("Pattern");
+
+	/** Delay */
 	static constexpr const TCHAR* DelayInDurationName = TEXT("DelayInDuration");
 	static constexpr const TCHAR* DelayOutDurationName = TEXT("DelayOutDuration");
 	static constexpr const TCHAR* DelaySpringFrequencyName = TEXT("DelaySpringFrequency");
@@ -416,28 +435,35 @@ struct FCEClonerEffectorChannelData
 		return Identifier;
 	}
 
-	/** General parameters */
-	ECEClonerEasing Easing = ECEClonerEasing::Linear;
-	ECEClonerEffectorMode Mode = ECEClonerEffectorMode::Default;
-	ECEClonerEffectorType Type = ECEClonerEffectorType::Sphere;
+	/** General */
 	float Magnitude = 0.f;
+	FVector Location = FVector::ZeroVector;
+	FQuat Rotation = FQuat::Identity;
+	FVector Scale = FVector::OneVector;
+	FLinearColor Color = FLinearColor::Transparent;
+
+	/** Type */
+	ECEClonerEffectorType Type = ECEClonerEffectorType::Sphere;
+	ECEClonerEasing Easing = ECEClonerEasing::Linear;
 	FVector InnerExtent = FVector::ZeroVector;
 	FVector OuterExtent = FVector::ZeroVector;
+
+	/** Mode */
+	ECEClonerEffectorMode Mode = ECEClonerEffectorMode::Default;
 	FVector LocationDelta = FVector::ZeroVector;
 	FVector RotationDelta = FVector::ZeroVector;
 	FVector ScaleDelta = FVector::OneVector;
-	FVector Location = FVector::ZeroVector;
-	FQuat Rotation = FQuat::Identity;
-	FVector Scale = FVector::OneVector;
 	float Frequency = 1.f;
 	FVector Pan = FVector::ZeroVector;
-	FLinearColor Color = FLinearColor::Transparent;
+	ECEClonerEffectorProceduralPattern Pattern = ECEClonerEffectorProceduralPattern::CurlNoise;
+
+	/** Delay */
 	float DelayInDuration = 0.f;
 	float DelayOutDuration = 0.f;
 	float DelaySpringFrequency = 0.f;
 	float DelaySpringFalloff = 0.f;
 
-	/** Forces parameters */
+	/** Forces */
 	float OrientationForceRate = 0.f;
 	FVector OrientationForceMin = FVector::ZeroVector;
 	FVector OrientationForceMax = FVector::ZeroVector;
diff --git a/Engine/Plugins/Experimental/ClonerEffector/Source/ClonerEffector/Public/Cloner/Layouts/CEClonerLayoutBase.h b/Engine/Plugins/Experimental/ClonerEffector/Source/ClonerEffector/Public/Cloner/Layouts/CEClonerLayoutBase.h
--- a/Engine/Plugins/Experimental/ClonerEffector/Source/ClonerEffector/Public/Cloner/Layouts/CEClonerLayoutBase.h
+++ b/Engine/Plugins/Experimental/ClonerEffector/Source/ClonerEffector/Public/Cloner/Layouts/CEClonerLayoutBase.h
@@ -24,6 +24,8 @@ enum class ECEClonerLayoutAssetVersion : uint8
 	DelayEffector,
 	/** Cancel clones effect added */
 	CancelEffector,
+	/** Procedural pattern effect added */
+	ProceduralEffector,
 	LatestVersionPlusOne,
 	LatestVersion = LatestVersionPlusOne - 1
 };
diff --git a/Engine/Plugins/Experimental/ClonerEffector/Source/ClonerEffector/Public/Effector/CEEffectorExtensionBase.h b/Engine/Plugins/Experimental/ClonerEffector/Source/ClonerEffector/Public/Effector/CEEffectorExtensionBase.h
--- a/Engine/Plugins/Experimental/ClonerEffector/Source/ClonerEffector/Public/Effector/CEEffectorExtensionBase.h
+++ b/Engine/Plugins/Experimental/ClonerEffector/Source/ClonerEffector/Public/Effector/CEEffectorExtensionBase.h
@@ -53,6 +53,12 @@ public:
 		return bExtensionActive;
 	}
 
+	/** Called after load when extension name is obsolete to handle redirection to this extension */
+	virtual bool RedirectExtensionName(FName InOldExtensionName) const
+	{
+		return false;
+	}
+
 protected:
 	//~ Begin UObject
 	virtual void PostEditImport() override;
diff --git a/Engine/Plugins/Experimental/ClonerEffector/Source/ClonerEffector/Public/Effector/Modes/CEEffectorNoiseMode.h b/Engine/Plugins/Experimental/ClonerEffector/Source/ClonerEffector/Public/Effector/Modes/CEEffectorNoiseMode.h
deleted file mode 100644
diff --git a/Engine/Plugins/Experimental/ClonerEffector/Source/ClonerEffector/Public/Effector/Modes/CEEffectorProceduralMode.h b/Engine/Plugins/Experimental/ClonerEffector/Source/ClonerEffector/Public/Effector/Modes/CEEffectorProceduralMode.h
new file mode 100644
--- /dev/null
+++ b/Engine/Plugins/Experimental/ClonerEffector/Source/ClonerEffector/Public/Effector/Modes/CEEffectorProceduralMode.h
@@ -0,0 +1,118 @@
+// Copyright Epic Games, Inc. All Rights Reserved.
+
+#pragma once
+
+#include "CEClonerEffectorShared.h"
+#include "CEEffectorModeBase.h"
+#include "CEPropertyChangeDispatcher.h"
+#include "CEEffectorProceduralMode.generated.h"
+
+class UCEEffectorComponent;
+
+UCLASS(MinimalAPI, BlueprintType, Within=CEEffectorComponent)
+class UCEEffectorProceduralMode : public UCEEffectorModeBase
+{
+	GENERATED_BODY()
+
+public:
+	UCEEffectorProceduralMode()
+		: UCEEffectorModeBase(TEXT("Procedural"), static_cast<int32>(ECEClonerEffectorMode::Procedural))
+	{}
+
+	UFUNCTION(BlueprintCallable, Category="Effector")
+	void SetPattern(ECEClonerEffectorProceduralPattern InPattern);
+
+	UFUNCTION(BlueprintPure, Category="Effector")
+	ECEClonerEffectorProceduralPattern GetPattern() const
+	{
+		return Pattern;
+	}
+
+	UFUNCTION(BlueprintCallable, Category="Effector")
+	CLONEREFFECTOR_API void SetLocationStrength(const FVector& InStrength);
+
+	UFUNCTION(BlueprintPure, Category="Effector")
+	FVector GetLocationStrength() const
+	{
+		return LocationStrength;
+	}
+
+	UFUNCTION(BlueprintCallable, Category="Effector")
+	CLONEREFFECTOR_API void SetRotationStrength(const FRotator& InStrength);
+
+	UFUNCTION(BlueprintPure, Category="Effector")
+	FRotator GetRotationStrength() const
+	{
+		return RotationStrength;
+	}
+
+	UFUNCTION(BlueprintCallable, Category="Effector")
+	CLONEREFFECTOR_API void SetScaleStrength(const FVector& InStrength);
+
+	UFUNCTION(BlueprintPure, Category="Effector")
+	FVector GetScaleStrength() const
+	{
+		return ScaleStrength;
+	}
+
+	UFUNCTION(BlueprintCallable, Category="Effector")
+	CLONEREFFECTOR_API void SetPan(const FVector& InPan);
+
+	UFUNCTION(BlueprintPure, Category="Effector")
+	FVector GetPan() const
+	{
+		return Pan;
+	}
+
+	UFUNCTION(BlueprintCallable, Category="Effector")
+	CLONEREFFECTOR_API void SetFrequency(float InFrequency);
+
+	UFUNCTION(BlueprintPure, Category="Effector")
+	float GetFrequency() const
+	{
+		return Frequency;
+	}
+
+protected:
+	//~ Begin UCEEffectorNoiseMode
+	virtual void OnExtensionParametersChanged(UCEEffectorComponent* InComponent) override;
+	virtual bool RedirectExtensionName(FName InOldExtensionName) const override;
+	//~ End UCEEffectorNoiseMode
+
+	//~ Begin UObject
+	virtual void PostInitProperties() override;
+#if WITH_EDITOR
+	virtual void PostEditChangeProperty(FPropertyChangedEvent& InPropertyChangedEvent) override;
+#endif
+	//~ End UObject
+
+	/** Procedural pattern behavior */
+	UPROPERTY(EditInstanceOnly, Setter, Getter, Category="Mode")
+	ECEClonerEffectorProceduralPattern Pattern = ECEClonerEffectorProceduralPattern::CurlNoise;
+
+	/** Amplitude of the noise field for location */
+	UPROPERTY(EditInstanceOnly, Setter, Getter, Category="Mode")
+	FVector LocationStrength = FVector::ZeroVector;
+
+	/** Amplitude of the noise field for rotation */
+	UPROPERTY(EditInstanceOnly, Setter, Getter, Category="Mode")
+	FRotator RotationStrength = FRotator::ZeroRotator;
+
+	/** Amplitude of the noise field for scale */
+	UPROPERTY(EditInstanceOnly, Setter, Getter, Category="Mode", meta=(ClampMin="0", MotionDesignVectorWidget, AllowPreserveRatio="XYZ", Delta="0.01"))
+	FVector ScaleStrength = FVector::OneVector;
+
+	/** Panning to offset the noise field sampling */
+	UPROPERTY(EditInstanceOnly, Setter, Getter, Category="Mode", meta=(MotionDesignVectorWidget, AllowPreserveRatio="XYZ"))
+	FVector Pan = FVector::ZeroVector;
+
+	/** Intensity of the noise field */
+	UPROPERTY(EditInstanceOnly, Setter, Getter, Category="Mode", meta=(ClampMin="0"))
+	float Frequency = 0.5f;
+
+private:
+#if WITH_EDITOR
+	/** Used for PECP */
+	static const TCEPropertyChangeDispatcher<UCEEffectorProceduralMode> PropertyChangeDispatcher;
+#endif
+};
\ No newline at end of file
