MotionDesign : ClonerEffector
- Texture contraint can now be clamped, offseted, scaled, rotated in grid layout

#virtualized

diff --git a/Engine/Plugins/Experimental/ClonerEffector/Source/ClonerEffector/Private/Cloner/Extensions/CEClonerConstraintExtension.cpp b/Engine/Plugins/Experimental/ClonerEffector/Source/ClonerEffector/Private/Cloner/Extensions/CEClonerConstraintExtension.cpp
--- a/Engine/Plugins/Experimental/ClonerEffector/Source/ClonerEffector/Private/Cloner/Extensions/CEClonerConstraintExtension.cpp
+++ b/Engine/Plugins/Experimental/ClonerEffector/Source/ClonerEffector/Private/Cloner/Extensions/CEClonerConstraintExtension.cpp
@@ -139,6 +139,50 @@ void UCEClonerConstraintExtension::SetTextureThreshold(float InThreshold)
 	MarkExtensionDirty();
 }
 
+void UCEClonerConstraintExtension::SetTextureOffset(const FVector2D& InOffset)
+{
+	if (TextureOffset.Equals(InOffset))
+	{
+		return;
+	}
+
+	TextureOffset = InOffset;
+	MarkExtensionDirty();
+}
+
+void UCEClonerConstraintExtension::SetTextureRotation(float InRotation)
+{
+	if (FMath::IsNearlyEqual(TextureRotation, InRotation))
+	{
+		return;
+	}
+
+	TextureRotation = InRotation;
+	MarkExtensionDirty();
+}
+
+void UCEClonerConstraintExtension::SetTextureScale(const FVector2D& InScale)
+{
+	if (TextureScale.Equals(InScale))
+	{
+		return;
+	}
+
+	TextureScale = InScale;
+	MarkExtensionDirty();
+}
+
+void UCEClonerConstraintExtension::SetTextureClamp(bool bInClamp)
+{
+	if (bTextureClamp == bInClamp)
+	{
+		return;
+	}
+
+	bTextureClamp = bInClamp;
+	MarkExtensionDirty();
+}
+
 #if WITH_EDITOR
 const TCEPropertyChangeDispatcher<UCEClonerConstraintExtension> UCEClonerConstraintExtension::PropertyChangeDispatcher =
 {
@@ -154,6 +198,10 @@ const TCEPropertyChangeDispatcher<UCEClonerConstraintExtension> UCEClonerConstra
 	{ GET_MEMBER_NAME_CHECKED(UCEClonerConstraintExtension, TexturePlane), &UCEClonerConstraintExtension::OnExtensionPropertyChanged },
 	{ GET_MEMBER_NAME_CHECKED(UCEClonerConstraintExtension, TextureCompareMode), &UCEClonerConstraintExtension::OnExtensionPropertyChanged },
 	{ GET_MEMBER_NAME_CHECKED(UCEClonerConstraintExtension, TextureThreshold), &UCEClonerConstraintExtension::OnExtensionPropertyChanged },
+	{ GET_MEMBER_NAME_CHECKED(UCEClonerConstraintExtension, TextureOffset), &UCEClonerConstraintExtension::OnExtensionPropertyChanged },
+	{ GET_MEMBER_NAME_CHECKED(UCEClonerConstraintExtension, TextureRotation), &UCEClonerConstraintExtension::OnExtensionPropertyChanged },
+	{ GET_MEMBER_NAME_CHECKED(UCEClonerConstraintExtension, TextureScale), &UCEClonerConstraintExtension::OnExtensionPropertyChanged },
+	{ GET_MEMBER_NAME_CHECKED(UCEClonerConstraintExtension, bTextureClamp), &UCEClonerConstraintExtension::OnExtensionPropertyChanged },
 };
 
 void UCEClonerConstraintExtension::PostEditChangeProperty(FPropertyChangedEvent& InPropertyChangedEvent)
@@ -201,9 +249,16 @@ void UCEClonerConstraintExtension::OnExtensionParametersChanged(UCEClonerCompone
 
 	static const FNiagaraVariable ConstraintTextureCompareModeVar(FNiagaraTypeDefinition(StaticEnum<ECEClonerCompareMode>()), TEXT("ConstraintTextureCompareMode"));
 	ExposedParameters.SetParameterValue<int32>(static_cast<int32>(TextureCompareMode), ConstraintTextureCompareModeVar);
+	
+	InComponent->SetFloatParameter(TEXT("ConstraintTextureThreshold"), FMath::Max(0, TextureThreshold));
+
+	InComponent->SetVariableVec2(TEXT("ConstraintTextureOffset"), TextureOffset);
+
+	InComponent->SetFloatParameter(TEXT("ConstraintTextureRotation"), TextureRotation);
+
+	InComponent->SetVariableVec2(TEXT("ConstraintTextureScale"), TextureScale);
 
-	// To avoid resaving system asset and overwriting 5.6 version, use typo name in 5.5 
-	InComponent->SetFloatParameter(TEXT("ContraintTextureThreshold"), FMath::Max(0, TextureThreshold));
+	InComponent->SetVariableBool(TEXT("ConstraintTextureClamp"), bTextureClamp);
 }
 
 bool UCEClonerConstraintExtension::IsLayoutSupported(const UCEClonerLayoutBase* InLayout) const
diff --git a/Engine/Plugins/Experimental/ClonerEffector/Source/ClonerEffector/Public/Cloner/Extensions/CEClonerConstraintExtension.h b/Engine/Plugins/Experimental/ClonerEffector/Source/ClonerEffector/Public/Cloner/Extensions/CEClonerConstraintExtension.h
--- a/Engine/Plugins/Experimental/ClonerEffector/Source/ClonerEffector/Public/Cloner/Extensions/CEClonerConstraintExtension.h
+++ b/Engine/Plugins/Experimental/ClonerEffector/Source/ClonerEffector/Public/Cloner/Extensions/CEClonerConstraintExtension.h
@@ -128,6 +128,42 @@ public:
 		return TextureThreshold;
 	}
 
+	UFUNCTION(BlueprintCallable, Category="Cloner")
+	CLONEREFFECTOR_API void SetTextureOffset(const FVector2D& InOffset);
+
+	UFUNCTION(BlueprintPure, Category="Cloner")
+	FVector2D GetTextureOffset() const
+	{
+		return TextureOffset;
+	}
+
+	UFUNCTION(BlueprintCallable, Category="Cloner")
+	CLONEREFFECTOR_API void SetTextureRotation(float InRotation);
+
+	UFUNCTION(BlueprintPure, Category="Cloner")
+	float GetTextureRotation() const
+	{
+		return TextureRotation;
+	}
+	
+	UFUNCTION(BlueprintCallable, Category="Cloner")
+	CLONEREFFECTOR_API void SetTextureScale(const FVector2D& InScale);
+
+	UFUNCTION(BlueprintPure, Category="Cloner")
+	FVector2D GetTextureScale() const
+	{
+		return TextureScale;
+	}
+	
+	UFUNCTION(BlueprintCallable, Category="Cloner")
+	CLONEREFFECTOR_API void SetTextureClamp(bool bInClamp);
+
+	UFUNCTION(BlueprintPure, Category="Cloner")
+	bool GetTextureClamp() const
+	{
+		return bTextureClamp;
+	}
+
 protected:
 	//~ Begin UObject
 #if WITH_EDITOR
@@ -182,6 +218,18 @@ protected:
 	UPROPERTY(EditInstanceOnly, Setter, Getter, Category="Constraint", meta=(EditCondition="Constraint == ECEClonerGridConstraint::Texture", EditConditionHides))
 	float TextureThreshold = 0.f;
 
+	UPROPERTY(EditInstanceOnly, Setter, Getter, Category="Constraint", meta=(EditCondition="Constraint == ECEClonerGridConstraint::Texture", EditConditionHides))
+	FVector2D TextureOffset = FVector2D::ZeroVector;
+	
+	UPROPERTY(EditInstanceOnly, Setter, Getter, Category="Constraint", meta=(EditCondition="Constraint == ECEClonerGridConstraint::Texture", EditConditionHides))
+	float TextureRotation = 0.f;
+	
+	UPROPERTY(EditInstanceOnly, Setter, Getter, Category="Constraint", meta=(MotionDesignVectorWidget, AllowPreserveRatio="XY", EditCondition="Constraint == ECEClonerGridConstraint::Texture", EditConditionHides))
+	FVector2D TextureScale = FVector2D(1.0);
+
+	UPROPERTY(EditInstanceOnly, Setter="SetTextureClamp", Getter="GetTextureClamp", Category="Constraint", meta=(EditCondition="Constraint == ECEClonerGridConstraint::Texture", EditConditionHides))
+	bool bTextureClamp = false;
+
 private:
 #if WITH_EDITOR
 	/** Used for PECP */
diff --git a/Engine/Plugins/Experimental/ClonerEffector/Source/ClonerEffector/Public/Cloner/Layouts/CEClonerLayoutBase.h b/Engine/Plugins/Experimental/ClonerEffector/Source/ClonerEffector/Public/Cloner/Layouts/CEClonerLayoutBase.h
--- a/Engine/Plugins/Experimental/ClonerEffector/Source/ClonerEffector/Public/Cloner/Layouts/CEClonerLayoutBase.h
+++ b/Engine/Plugins/Experimental/ClonerEffector/Source/ClonerEffector/Public/Cloner/Layouts/CEClonerLayoutBase.h
@@ -26,6 +26,8 @@ enum class ECEClonerLayoutAssetVersion : uint8
 	CancelEffector,
 	/** Procedural pattern effect added */
 	ProceduralEffector,
+	/** Texture constraint improved */
+	TextureConstraint,
 	LatestVersionPlusOne,
 	LatestVersion = LatestVersionPlusOne - 1
 };
