MotionDesign : ClonerEffector
- Add concentric support in circle layout by specifying RingCount

#virtualized

diff --git a/Engine/Plugins/Experimental/ClonerEffector/Source/ClonerEffector/Private/Cloner/Layouts/CEClonerCircleLayout.cpp b/Engine/Plugins/Experimental/ClonerEffector/Source/ClonerEffector/Private/Cloner/Layouts/CEClonerCircleLayout.cpp
--- a/Engine/Plugins/Experimental/ClonerEffector/Source/ClonerEffector/Private/Cloner/Layouts/CEClonerCircleLayout.cpp
+++ b/Engine/Plugins/Experimental/ClonerEffector/Source/ClonerEffector/Private/Cloner/Layouts/CEClonerCircleLayout.cpp
@@ -15,6 +15,18 @@ void UCEClonerCircleLayout::SetCount(int32 InCount)
 	MarkLayoutDirty();
 }
 
+void UCEClonerCircleLayout::SetRingCount(int32 InCount)
+{
+	InCount = FMath::Max(1, InCount);
+	if (RingCount == InCount)
+	{
+		return;
+	}
+
+	RingCount = InCount;
+	MarkLayoutDirty();
+}
+
 void UCEClonerCircleLayout::SetRadius(float InRadius)
 {
 	if (Radius == InRadius)
@@ -96,6 +108,7 @@ void UCEClonerCircleLayout::SetScale(const FVector& InScale)
 const TCEPropertyChangeDispatcher<UCEClonerCircleLayout> UCEClonerCircleLayout::PropertyChangeDispatcher =
 {
 	{ GET_MEMBER_NAME_CHECKED(UCEClonerCircleLayout, Count), &UCEClonerCircleLayout::OnLayoutPropertyChanged },
+	{ GET_MEMBER_NAME_CHECKED(UCEClonerCircleLayout, RingCount), &UCEClonerCircleLayout::OnLayoutPropertyChanged },
 	{ GET_MEMBER_NAME_CHECKED(UCEClonerCircleLayout, Radius), &UCEClonerCircleLayout::OnLayoutPropertyChanged },
 	{ GET_MEMBER_NAME_CHECKED(UCEClonerCircleLayout, AngleStart), &UCEClonerCircleLayout::OnLayoutPropertyChanged },
 	{ GET_MEMBER_NAME_CHECKED(UCEClonerCircleLayout, AngleRatio), &UCEClonerCircleLayout::OnLayoutPropertyChanged },
@@ -119,6 +132,8 @@ void UCEClonerCircleLayout::OnLayoutParametersChanged(UCEClonerComponent* InComp
 
 	InComponent->SetIntParameter(TEXT("CircleCount"), Count);
 
+	InComponent->SetIntParameter(TEXT("CircleRingCount"), RingCount);
+
 	InComponent->SetFloatParameter(TEXT("CircleRadius"), Radius);
 
 	InComponent->SetFloatParameter(TEXT("CircleStart"), AngleStart);
diff --git a/Engine/Plugins/Experimental/ClonerEffector/Source/ClonerEffector/Public/Cloner/Layouts/CEClonerCircleLayout.h b/Engine/Plugins/Experimental/ClonerEffector/Source/ClonerEffector/Public/Cloner/Layouts/CEClonerCircleLayout.h
--- a/Engine/Plugins/Experimental/ClonerEffector/Source/ClonerEffector/Public/Cloner/Layouts/CEClonerCircleLayout.h
+++ b/Engine/Plugins/Experimental/ClonerEffector/Source/ClonerEffector/Public/Cloner/Layouts/CEClonerCircleLayout.h
@@ -31,6 +31,15 @@ public:
 		return Count;
 	}
 
+	UFUNCTION(BlueprintCallable, Category="Cloner|Layout|Circle")
+	CLONEREFFECTOR_API void SetRingCount(int32 InCount);
+
+	UFUNCTION(BlueprintPure, Category="Cloner|Layout|Circle")
+	int32 GetRingCount() const
+	{
+		return RingCount;
+	}
+
 	UFUNCTION(BlueprintCallable, Category="Cloner|Layout|Circle")
 	CLONEREFFECTOR_API void SetRadius(float InRadius);
 
@@ -108,6 +117,9 @@ protected:
 	UPROPERTY(EditInstanceOnly, Setter, Getter, Category="Layout")
 	int32 Count = 3 * 3;
 
+	UPROPERTY(EditInstanceOnly, Setter, Getter, Category="Layout", meta=(ClampMin="1"))
+	int32 RingCount = 1;
+
 	UPROPERTY(EditInstanceOnly, Setter, Getter, Category="Layout")
 	float Radius = 200.f;
 
diff --git a/Engine/Plugins/Experimental/ClonerEffector/Source/ClonerEffector/Public/Cloner/Layouts/CEClonerLayoutBase.h b/Engine/Plugins/Experimental/ClonerEffector/Source/ClonerEffector/Public/Cloner/Layouts/CEClonerLayoutBase.h
--- a/Engine/Plugins/Experimental/ClonerEffector/Source/ClonerEffector/Public/Cloner/Layouts/CEClonerLayoutBase.h
+++ b/Engine/Plugins/Experimental/ClonerEffector/Source/ClonerEffector/Public/Cloner/Layouts/CEClonerLayoutBase.h
@@ -28,6 +28,8 @@ enum class ECEClonerLayoutAssetVersion : uint8
 	ProceduralEffector,
 	/** Texture constraint improved */
 	TextureConstraint,
+	/** circle ring added */
+	ConcentricCircle,
 	LatestVersionPlusOne,
 	LatestVersion = LatestVersionPlusOne - 1
 };
