MotionDesign : ClonerEffector
- Due to dynamic to static mesh conversion not working at runtime, cloner baked static meshes will be cached in editor to allow cloner/effector usage at runtime

diff --git a/Engine/Plugins/Experimental/ClonerEffector/Source/ClonerEffector/Private/Cloner/CEClonerComponent.cpp b/Engine/Plugins/Experimental/ClonerEffector/Source/ClonerEffector/Private/Cloner/CEClonerComponent.cpp
--- a/Engine/Plugins/Experimental/ClonerEffector/Source/ClonerEffector/Private/Cloner/CEClonerComponent.cpp
+++ b/Engine/Plugins/Experimental/ClonerEffector/Source/ClonerEffector/Private/Cloner/CEClonerComponent.cpp
@@ -64,11 +64,11 @@ UCEClonerComponent::UCEClonerComponent()
 	if (!IsTemplate())
 	{
 		UCEClonerSubsystem::OnClonerSetEnabled().AddUObject(this, &UCEClonerComponent::OnClonerSetEnabled);
-		USceneComponent::MarkRenderStateDirtyEvent.AddUObject(this, &UCEClonerComponent::OnRenderStateDirty);
 
 		// Bind to delegate to detect material changes
 #if WITH_EDITOR
-		FCoreUObjectDelegates::OnObjectPropertyChanged.RemoveAll(this);
+		USceneComponent::MarkRenderStateDirtyEvent.AddUObject(this, &UCEClonerComponent::OnRenderStateDirty);
+
 		FCoreUObjectDelegates::OnObjectPropertyChanged.AddUObject(this, &UCEClonerComponent::OnActorPropertyChanged);
 
 		UMaterial::OnMaterialCompilationFinished().AddUObject(this, &UCEClonerComponent::OnMaterialCompiled);
@@ -179,6 +179,10 @@ void UCEClonerComponent::OnComponentDestroyed(bool bInDestroyingHierarchy)
 	Super::OnComponentDestroyed(bInDestroyingHierarchy);
 
 #if WITH_EDITOR
+	USceneComponent::MarkRenderStateDirtyEvent.RemoveAll(this);
+	
+	FCoreUObjectDelegates::OnObjectPropertyChanged.RemoveAll(this);
+	
 	UMaterial::OnMaterialCompilationFinished().RemoveAll(this);
 #endif
 }
@@ -209,7 +213,6 @@ void UCEClonerComponent::UpdateClonerRenderState()
 
 void UCEClonerComponent::UpdateClonerAttachmentTree(bool bInReset)
 {
-#if WITH_EDITOR
 	if (ClonerTree.Status == ECEClonerAttachmentStatus::Updated)
 	{
 		ClonerTree.Status = ECEClonerAttachmentStatus::Outdated;
@@ -222,7 +225,6 @@ void UCEClonerComponent::UpdateClonerAttachmentTree(bool bInReset)
 	}
 
 	UpdateAttachmentTree();
-#endif
 }
 
 void UCEClonerComponent::UpdateAttachmentTree()
@@ -425,7 +427,6 @@ void UCEClonerComponent::BindActorDelegates(AActor* InActor)
 			StaticMeshComponent->OnStaticMeshChanged().AddUObject(this, &UCEClonerComponent::OnMeshChanged, InActor);
 		}
 	}
-#endif
 
 	// Detect dynamic mesh change
 	TArray<UDynamicMeshComponent*> DynamicMeshComponents;
@@ -438,6 +439,7 @@ void UCEClonerComponent::BindActorDelegates(AActor* InActor)
 			DynamicMeshComponent->OnMeshChanged.AddUObject(this, &UCEClonerComponent::OnMeshChanged, NullComponent, InActor);
 		}
 	}
+#endif
 
 	// Detect components transform
 	TArray<USceneComponent*> SceneComponents;
@@ -467,7 +469,6 @@ void UCEClonerComponent::UnbindActorDelegates(AActor* InActor) const
 	{
 		StaticMeshComponent->OnStaticMeshChanged().RemoveAll(this);
 	}
-#endif
 
 	TArray<UDynamicMeshComponent*> DynamicMeshComponents;
 	InActor->GetComponents(DynamicMeshComponents, false);
@@ -475,6 +476,7 @@ void UCEClonerComponent::UnbindActorDelegates(AActor* InActor) const
 	{
 		DynamicMeshComponent->OnMeshChanged.RemoveAll(this);
 	}
+#endif
 
 	TArray<USceneComponent*> SceneComponents;
 	InActor->GetComponents(SceneComponents, /** IncludeChildren */false);
@@ -904,11 +906,15 @@ void UCEClonerComponent::UpdateRootActorBakedStaticMesh(AActor* InRootActor)
 		MeshBuilder.AppendMesh(AttachmentItem->BakedMesh, AttachmentItem->BakedMaterials, MeshTransform);
 	}
 
+	// Avoid dirtying asset
 	UStaticMesh* Mesh = NewObject<UStaticMesh>();
+
 	TArray<TWeakObjectPtr<UMaterialInterface>> MeshMaterials;
 	bClonerMeshesDirty = MeshBuilder.BuildStaticMesh(Mesh, MeshMaterials);
 	MeshBuilder.Reset();
 
+	Mesh->Rename(nullptr, this, REN_NonTransactional | REN_DoNotDirty);
+
 	ClonerTree.MergedBakedMeshes[RootIdx] = Mesh;
 }
 
diff --git a/Engine/Plugins/Experimental/ClonerEffector/Source/ClonerEffector/Public/CEClonerEffectorShared.h b/Engine/Plugins/Experimental/ClonerEffector/Source/ClonerEffector/Public/CEClonerEffectorShared.h
--- a/Engine/Plugins/Experimental/ClonerEffector/Source/ClonerEffector/Public/CEClonerEffectorShared.h
+++ b/Engine/Plugins/Experimental/ClonerEffector/Source/ClonerEffector/Public/CEClonerEffectorShared.h
@@ -310,11 +310,11 @@ struct FCEClonerAttachmentItem
 
 	/** Actual baked mesh for the current actor item,
 		dynamic because it's easier for merging them together and avoid conversions again */
-	UPROPERTY()
+	UPROPERTY(Transient)
 	TObjectPtr<UDynamicMesh> BakedMesh;
 
 	/** Actual baked materials for the current actor item */
-	UPROPERTY()
+	UPROPERTY(Transient)
 	TArray<TWeakObjectPtr<UMaterialInterface>> BakedMaterials;
 
 	/** Status of the baked mesh, does it needs to be updated */
@@ -363,11 +363,11 @@ struct FCEClonerAttachmentTree
 	TArray<TObjectPtr<UStaticMesh>> MergedBakedMeshes;
 
 	/** Status of the cloner tree */
-	UPROPERTY()
+	UPROPERTY(Transient)
 	ECEClonerAttachmentStatus Status = ECEClonerAttachmentStatus::Updated;
 
 	/** Attachment items that are dirty and need an update */
-	UPROPERTY()
+	UPROPERTY(Transient)
 	TSet<TWeakObjectPtr<AActor>> DirtyItemAttachments;
 
 	void Reset()
diff --git a/Engine/Plugins/Experimental/ClonerEffector/Source/ClonerEffector/Public/Cloner/CEClonerComponent.h b/Engine/Plugins/Experimental/ClonerEffector/Source/ClonerEffector/Public/Cloner/CEClonerComponent.h
--- a/Engine/Plugins/Experimental/ClonerEffector/Source/ClonerEffector/Public/Cloner/CEClonerComponent.h
+++ b/Engine/Plugins/Experimental/ClonerEffector/Source/ClonerEffector/Public/Cloner/CEClonerComponent.h
@@ -434,7 +434,7 @@ private:
 	TArray<FName> GetClonerLayoutNames() const;
 
 	/** Attachment tree view */
-	UPROPERTY(Transient, DuplicateTransient, NonTransactional)
+	UPROPERTY(NonTransactional)
 	FCEClonerAttachmentTree ClonerTree;
 
 	UPROPERTY(Transient, DuplicateTransient, NonTransactional)
diff --git a/Engine/Plugins/Experimental/ClonerEffector/Source/ClonerEffectorMeshBuilder/ClonerEffectorMeshBuilder.Build.cs b/Engine/Plugins/Experimental/ClonerEffector/Source/ClonerEffectorMeshBuilder/ClonerEffectorMeshBuilder.Build.cs
--- a/Engine/Plugins/Experimental/ClonerEffector/Source/ClonerEffectorMeshBuilder/ClonerEffectorMeshBuilder.Build.cs
+++ b/Engine/Plugins/Experimental/ClonerEffector/Source/ClonerEffectorMeshBuilder/ClonerEffectorMeshBuilder.Build.cs
@@ -24,6 +24,9 @@ public class ClonerEffectorMeshBuilder : ModuleRules
                 "GeometryCore",
                 "GeometryFramework",
                 "GeometryScriptingCore",
+                "MeshConversion",
+				"MeshDescription",
+				"ModelingComponents",
                 "Niagara",
                 "NiagaraCore",
                 "ProceduralMeshComponent",
diff --git a/Engine/Plugins/Experimental/ClonerEffector/Source/ClonerEffectorMeshBuilder/Private/CEMeshBuilder.cpp b/Engine/Plugins/Experimental/ClonerEffector/Source/ClonerEffectorMeshBuilder/Private/CEMeshBuilder.cpp
--- a/Engine/Plugins/Experimental/ClonerEffector/Source/ClonerEffectorMeshBuilder/Private/CEMeshBuilder.cpp
+++ b/Engine/Plugins/Experimental/ClonerEffector/Source/ClonerEffectorMeshBuilder/Private/CEMeshBuilder.cpp
@@ -2,6 +2,7 @@
 
 #include "CEMeshBuilder.h"
 
+#include "AssetUtils/StaticMeshMaterialUtil.h"
 #include "Components/BrushComponent.h"
 #include "Components/DynamicMeshComponent.h"
 #include "Components/InstancedStaticMeshComponent.h"
@@ -11,6 +12,7 @@
 #include "DynamicMesh/MeshTransforms.h"
 #include "DynamicMesh/Operations/MergeCoincidentMeshEdges.h"
 #include "DynamicMeshEditor.h"
+#include "DynamicMeshToMeshDescription.h"
 #include "Engine/SkeletalMesh.h"
 #include "Engine/StaticMesh.h"
 #include "GeometryScript/MeshAssetFunctions.h"
@@ -25,7 +27,7 @@
 #include "NiagaraSystem.h"
 #include "ProceduralMeshComponent.h"
 #include "Rendering/SkeletalMeshRenderData.h"
-#include "StaticMeshOperations.h"
+#include "StaticMeshDescription.h"
 #include "UDynamicMesh.h"
 
 const FCEMeshBuilder::FCEMeshBuilderParams FCEMeshBuilder::DefaultBuildParams;
@@ -1026,40 +1028,82 @@ bool FCEMeshBuilder::AppendPrimitiveComponent(const UObject* InMeshObject, UPrim
 
 bool FCEMeshBuilder::DynamicMeshToStaticMesh(UDynamicMesh* InMesh, UStaticMesh* OutMesh, const TArray<TWeakObjectPtr<UMaterialInterface>>& InMaterials)
 {
-	TArray<TObjectPtr<UMaterialInterface>> NewMaterials;
-	Algo::Transform(InMaterials, NewMaterials, [](const TWeakObjectPtr<UMaterialInterface>& InMaterialWeak)
+	using namespace UE::Geometry;
+
+	bool bResult = false;
+	
+	if (!InMesh || !OutMesh)
+    {
+        return bResult;
+    }
+
+	UStaticMeshDescription* StaticMeshDescription = OutMesh->GetStaticMeshDescription(0);
+
+	if (!StaticMeshDescription)
+	{
+		StaticMeshDescription = UStaticMesh::CreateStaticMeshDescription(OutMesh);
+	}
+
+	if (!StaticMeshDescription)
 	{
-		return InMaterialWeak.Get();
+		return bResult;
+	}
+
+	// Mesh
+	FMeshDescription& MeshDescription = StaticMeshDescription->GetMeshDescription();
+
+	InMesh->ProcessMesh([&MeshDescription](const FDynamicMesh3& InSourceMesh)
+	{
+		FConversionToMeshDescriptionOptions ConversionOptions;
+		ConversionOptions.bUpdateTangents = true;
+		ConversionOptions.bUpdateUVs = true;
+
+		FDynamicMeshToMeshDescription Converter(ConversionOptions);
+		Converter.Convert(&InSourceMesh, MeshDescription, /** CopyTangents */true);
 	});
 
-	// Export options
-	FGeometryScriptCopyMeshToAssetOptions AssetOptions;
-	AssetOptions.bReplaceMaterials = true;
-	AssetOptions.bEnableRecomputeNormals = false;
-	AssetOptions.bEnableRecomputeTangents = false;
-	AssetOptions.bEnableRemoveDegenerates = true;
-	AssetOptions.bEmitTransaction = false;
-	AssetOptions.bApplyNaniteSettings = false;
-	AssetOptions.bDeferMeshPostEditChange = false;
-	AssetOptions.NewMaterials = NewMaterials;
-
-	// LOD options
-	FGeometryScriptMeshWriteLOD TargetLOD;
-	TargetLOD.LODIndex = 0;
-	TargetLOD.bWriteHiResSource = false;
-
-	// Convert merged mesh to static mesh
-	EGeometryScriptOutcomePins OutResult;
-	UGeometryScriptLibrary_StaticMeshFunctions::CopyMeshToStaticMesh(InMesh, OutMesh, AssetOptions, TargetLOD, OutResult);
-
-#if WITH_EDITORONLY_DATA
-	// Compute normals and tangents
-	FMeshDescription* MeshDescription = OutMesh->GetMeshDescription(TargetLOD.LODIndex);
-	FStaticMeshOperations::ComputeTriangleTangentsAndNormals(*MeshDescription);
-	FStaticMeshOperations::ComputeTangentsAndNormals(*MeshDescription, EComputeNTBsFlags::Normals);
+	// Materials
+	TArray<FStaticMaterial> NewStaticMaterials;
+	NewStaticMaterials.Reserve(InMaterials.Num());
+	for (int32 MaterialIndex = 0; MaterialIndex < InMaterials.Num(); ++MaterialIndex)
+	{
+		FStaticMaterial& NewMaterial = NewStaticMaterials.AddDefaulted_GetRef();
+		NewMaterial.MaterialInterface = InMaterials[MaterialIndex].Get();
+		NewMaterial.MaterialSlotName = UE::AssetUtils::GenerateNewMaterialSlotName(NewStaticMaterials, NewMaterial.MaterialInterface, MaterialIndex);
+		NewMaterial.ImportedMaterialSlotName = NewMaterial.MaterialSlotName;
+		NewMaterial.UVChannelData = FMeshUVChannelInfo(1.f);
+	}
+
+	OutMesh->SetStaticMaterials(NewStaticMaterials);
+
+	TPolygonGroupAttributesRef<FName> PolygonGroupImportedMaterialSlotNames = StaticMeshDescription->GetRequiredAttributes().GetPolygonGroupMaterialSlotNames();
+
+	for (int32 SlotIdx = 0; SlotIdx < NewStaticMaterials.Num(); ++SlotIdx)
+	{
+		if (SlotIdx < PolygonGroupImportedMaterialSlotNames.GetNumElements())
+		{
+			PolygonGroupImportedMaterialSlotNames.Set(SlotIdx, NewStaticMaterials[SlotIdx].MaterialSlotName);
+		}
+	}
+
+	// Build
+	UStaticMesh::FBuildMeshDescriptionsParams Params;
+	Params.bFastBuild = true;
+	Params.bAllowCpuAccess = true;
+	Params.bCommitMeshDescription = true;
+
+	const TArray<const FMeshDescription*> MeshDescriptions {&MeshDescription};
+	bResult = OutMesh->BuildFromMeshDescriptions(MeshDescriptions, Params);
+
+#if WITH_EDITOR
+	UStaticMesh::FBuildParameters BuildParameters;
+	BuildParameters.bInSilent = true;
+	BuildParameters.bInRebuildUVChannelData = true;
+	BuildParameters.bInEnforceLightmapRestrictions = true;
+	OutMesh->Build(BuildParameters);
 #endif
 
-	return OutResult == EGeometryScriptOutcomePins::Success;
+	return bResult;
 }
 
 void FCEMeshBuilder::ClearOutputMesh() const
