MotionDesign : ClonerEffector
- Local and world space support in editor for the simulation, simulation is then cached for runtime usage
- Adjusting systems, modules, functions to support local or world space coordinates
- Mesh and Spline layout will not follow referenced actor anymore
- Simplified spawn options for instant or rate

#virtualized

diff --git a/Engine/Plugins/Experimental/ClonerEffector/Source/ClonerEffector/Private/Cloner/Extensions/CEClonerEmitterSpawnExtension.cpp b/Engine/Plugins/Experimental/ClonerEffector/Source/ClonerEffector/Private/Cloner/Extensions/CEClonerEmitterSpawnExtension.cpp
--- a/Engine/Plugins/Experimental/ClonerEffector/Source/ClonerEffector/Private/Cloner/Extensions/CEClonerEmitterSpawnExtension.cpp
+++ b/Engine/Plugins/Experimental/ClonerEffector/Source/ClonerEffector/Private/Cloner/Extensions/CEClonerEmitterSpawnExtension.cpp
@@ -2,6 +2,7 @@
 
 #include "Cloner/Extensions/CEClonerEmitterSpawnExtension.h"
 
+#include "NiagaraEmitter.h"
 #include "Cloner/CEClonerComponent.h"
 #include "NiagaraSystem.h"
 #include "NiagaraUserRedirectionParameterStore.h"
@@ -19,32 +20,24 @@ void UCEClonerEmitterSpawnExtension::SetSpawnLoopMode(ECEClonerSpawnLoopMode InM
 
 void UCEClonerEmitterSpawnExtension::SetSpawnLoopIterations(int32 InIterations)
 {
+	InIterations = FMath::Max(InIterations, 1);
 	if (SpawnLoopIterations == InIterations)
 	{
 		return;
 	}
 
-	if (InIterations < 1)
-	{
-		return;
-	}
-
 	SpawnLoopIterations = InIterations;
 	MarkExtensionDirty();
 }
 
 void UCEClonerEmitterSpawnExtension::SetSpawnLoopInterval(float InInterval)
 {
+	InInterval = FMath::Max(InInterval, 0.f);
 	if (SpawnLoopInterval == InInterval)
 	{
 		return;
 	}
 
-	if (InInterval < 0.f)
-	{
-		return;
-	}
-
 	SpawnLoopInterval = InInterval;
 	MarkExtensionDirty();
 }
@@ -62,16 +55,12 @@ void UCEClonerEmitterSpawnExtension::SetSpawnBehaviorMode(ECEClonerSpawnBehavior
 
 void UCEClonerEmitterSpawnExtension::SetSpawnRate(float InRate)
 {
+	InRate = FMath::Max(InRate, 0.f);
 	if (SpawnRate == InRate)
 	{
 		return;
 	}
 
-	if (InRate < 0)
-	{
-		return;
-	}
-
 	SpawnRate = InRate;
 	MarkExtensionDirty();
 }
@@ -85,18 +74,39 @@ void UCEClonerEmitterSpawnExtension::OnExtensionParametersChanged(UCEClonerCompo
 	const FNiagaraVariable SpawnLoopModeVar(FNiagaraTypeDefinition(StaticEnum<ECEClonerSpawnLoopMode>()), TEXT("SpawnLoopMode"));
 	ExposedParameters.SetParameterValue<int32>(static_cast<int32>(SpawnLoopMode), SpawnLoopModeVar);
 
+	InComponent->SetIntParameter(TEXT("SpawnLoopIterations"), SpawnLoopIterations);
+	
 	const ECEClonerSpawnBehaviorMode BehaviorMode = SpawnLoopMode == ECEClonerSpawnLoopMode::Once ? ECEClonerSpawnBehaviorMode::Instant : SpawnBehaviorMode;
 	const FNiagaraVariable SpawnBehaviorModeVar(FNiagaraTypeDefinition(StaticEnum<ECEClonerSpawnBehaviorMode>()), TEXT("SpawnBehaviorMode"));
 	ExposedParameters.SetParameterValue<int32>(static_cast<int32>(BehaviorMode), SpawnBehaviorModeVar);
-
-	InComponent->SetFloatParameter(TEXT("SpawnLoopInterval"), SpawnLoopInterval);
-
-	InComponent->SetIntParameter(TEXT("SpawnLoopIterations"), SpawnLoopIterations);
+	
+	InComponent->SetFloatParameter(TEXT("SpawnLoopInterval"), BehaviorMode == ECEClonerSpawnBehaviorMode::Instant ? SpawnLoopInterval : 1.f);
 
 	InComponent->SetFloatParameter(TEXT("SpawnRate"), SpawnRate);
+
+#if WITH_EDITOR
+	// Do not allow world space when spawning once
+	if (SpawnLoopMode == ECEClonerSpawnLoopMode::Once && !bUseLocalSpace)
+	{
+		bUseLocalSpace = true;
+	}
+	
+	OnLocalSpaceChanged();
+#endif
 }
 
 #if WITH_EDITOR
+void UCEClonerEmitterSpawnExtension::SetUseLocalSpace(bool bInLocalSpace)
+{
+	if (bUseLocalSpace == bInLocalSpace)
+	{
+		return;
+	}
+
+	bUseLocalSpace = bInLocalSpace;
+	MarkExtensionDirty();
+}
+
 const TCEPropertyChangeDispatcher<UCEClonerEmitterSpawnExtension> UCEClonerEmitterSpawnExtension::PropertyChangeDispatcher =
 {
 	/** Spawn */
@@ -105,6 +115,7 @@ const TCEPropertyChangeDispatcher<UCEClonerEmitterSpawnExtension> UCEClonerEmitt
 	{ GET_MEMBER_NAME_CHECKED(UCEClonerEmitterSpawnExtension, SpawnLoopIterations), &UCEClonerEmitterSpawnExtension::OnExtensionPropertyChanged },
 	{ GET_MEMBER_NAME_CHECKED(UCEClonerEmitterSpawnExtension, SpawnBehaviorMode), &UCEClonerEmitterSpawnExtension::OnExtensionPropertyChanged },
 	{ GET_MEMBER_NAME_CHECKED(UCEClonerEmitterSpawnExtension, SpawnRate), &UCEClonerEmitterSpawnExtension::OnExtensionPropertyChanged },
+	{ GET_MEMBER_NAME_CHECKED(UCEClonerEmitterSpawnExtension, bUseLocalSpace), &UCEClonerEmitterSpawnExtension::OnExtensionPropertyChanged },
 };
 
 void UCEClonerEmitterSpawnExtension::PostEditChangeProperty(FPropertyChangedEvent& InPropertyChangedEvent)
@@ -113,4 +124,32 @@ void UCEClonerEmitterSpawnExtension::PostEditChangeProperty(FPropertyChangedEven
 
 	PropertyChangeDispatcher.OnPropertyChanged(this, InPropertyChangedEvent);
 }
-#endif
+
+void UCEClonerEmitterSpawnExtension::OnLocalSpaceChanged()
+{
+	UCEClonerComponent* ClonerComponent = GetClonerComponent();
+
+	if (!ClonerComponent)
+	{
+		return;
+	}
+
+	UNiagaraSystem* ClonerSystem = ClonerComponent->GetAsset();
+
+	if (!ClonerSystem || ClonerSystem->GetNumEmitters() == 0)
+	{
+		return;
+	}
+	
+	const FNiagaraEmitterHandle& EmitterHandle = ClonerSystem->GetEmitterHandle(0);
+
+	if (FVersionedNiagaraEmitterData* EmitterData = EmitterHandle.GetEmitterData())
+	{
+		if (EmitterData->bLocalSpace != bUseLocalSpace)
+		{
+			EmitterData->bLocalSpace = bUseLocalSpace;
+			ClonerSystem->RequestCompile(/** Force */false);
+		}
+	}
+}
+#endif
\ No newline at end of file
diff --git a/Engine/Plugins/Experimental/ClonerEffector/Source/ClonerEffector/Private/Cloner/Layouts/CEClonerMeshLayout.cpp b/Engine/Plugins/Experimental/ClonerEffector/Source/ClonerEffector/Private/Cloner/Layouts/CEClonerMeshLayout.cpp
--- a/Engine/Plugins/Experimental/ClonerEffector/Source/ClonerEffector/Private/Cloner/Layouts/CEClonerMeshLayout.cpp
+++ b/Engine/Plugins/Experimental/ClonerEffector/Source/ClonerEffector/Private/Cloner/Layouts/CEClonerMeshLayout.cpp
@@ -87,17 +87,6 @@ void UCEClonerMeshLayout::PostEditChangeProperty(FPropertyChangedEvent& InProper
 }
 #endif
 
-void UCEClonerMeshLayout::OnLayoutInactive()
-{
-	Super::OnLayoutInactive();
-
-	// unbind
-	if (USceneComponent* SceneComponent = SceneComponentWeak.Get())
-	{
-		SceneComponent->TransformUpdated.RemoveAll(this);
-	}
-}
-
 void UCEClonerMeshLayout::OnLayoutParametersChanged(UCEClonerComponent* InComponent)
 {
 	Super::OnLayoutParametersChanged(InComponent);
@@ -116,10 +105,8 @@ void UCEClonerMeshLayout::OnLayoutParametersChanged(UCEClonerComponent* InCompon
 	UNiagaraDataInterfaceActorComponent* ActorMeshDI = Cast<UNiagaraDataInterfaceActorComponent>(ExposedParameters.GetDataInterface(SampleMeshActorVar));
 
 	// unbind
-	if (USceneComponent* SceneComponent = SceneComponentWeak.Get())
+	if (const USceneComponent* SceneComponent = SceneComponentWeak.Get())
 	{
-		SceneComponent->TransformUpdated.RemoveAll(this);
-
 		if (AActor* Actor = SceneComponent->GetOwner())
 		{
 			Actor->OnDestroyed.RemoveAll(this);
@@ -133,9 +120,7 @@ void UCEClonerMeshLayout::OnLayoutParametersChanged(UCEClonerComponent* InCompon
 	{
 		ActorMeshDI->SourceActor = SampleActor;
 		SampleActor->OnDestroyed.AddUniqueDynamic(this, &UCEClonerMeshLayout::OnSampleActorDestroyed);
-
 		SceneComponentWeak = SampleActor->GetRootComponent();
-		SceneComponentWeak->TransformUpdated.AddUObject(this, &UCEClonerMeshLayout::OnSampleMeshTransformed);
 	}
 	else
 	{
@@ -170,11 +155,6 @@ void UCEClonerMeshLayout::OnLayoutParametersChanged(UCEClonerComponent* InCompon
 	}
 }
 
-void UCEClonerMeshLayout::OnSampleMeshTransformed(USceneComponent* InComponent, EUpdateTransformFlags InFlags, ETeleportType InType)
-{
-	MarkLayoutDirty();
-}
-
 void UCEClonerMeshLayout::OnSampleActorDestroyed(AActor* InDestroyedActor)
 {
 	UCEClonerComponent* ClonerComponent = GetClonerComponent();
diff --git a/Engine/Plugins/Experimental/ClonerEffector/Source/ClonerEffector/Private/Cloner/Layouts/CEClonerSplineLayout.cpp b/Engine/Plugins/Experimental/ClonerEffector/Source/ClonerEffector/Private/Cloner/Layouts/CEClonerSplineLayout.cpp
--- a/Engine/Plugins/Experimental/ClonerEffector/Source/ClonerEffector/Private/Cloner/Layouts/CEClonerSplineLayout.cpp
+++ b/Engine/Plugins/Experimental/ClonerEffector/Source/ClonerEffector/Private/Cloner/Layouts/CEClonerSplineLayout.cpp
@@ -122,16 +122,18 @@ void UCEClonerSplineLayout::PostEditChangeProperty(FPropertyChangedEvent& InProp
 }
 #endif
 
+void UCEClonerSplineLayout::OnLayoutActive()
+{
+	Super::OnLayoutActive();
+
+	USceneComponent::MarkRenderStateDirtyEvent.AddUObject(this, &UCEClonerSplineLayout::OnSampleSplineRenderStateUpdated);
+}
+
 void UCEClonerSplineLayout::OnLayoutInactive()
 {
 	Super::OnLayoutInactive();
 
-	// unbind
-	if (USplineComponent* SplineComponent = SplineComponentWeak.Get())
-	{
-		SplineComponent->TransformUpdated.RemoveAll(this);
-		USceneComponent::MarkRenderStateDirtyEvent.RemoveAll(this);
-	}
+	USceneComponent::MarkRenderStateDirtyEvent.RemoveAll(this);
 }
 
 void UCEClonerSplineLayout::OnLayoutParametersChanged(UCEClonerComponent* InComponent)
@@ -147,44 +149,24 @@ void UCEClonerSplineLayout::OnLayoutParametersChanged(UCEClonerComponent* InComp
 	InComponent->SetBoolParameter(TEXT("MeshOrientAxisEnable"), bOrientMesh);
 
 	// unbind
+	if (AActor* PrevSplineActor = SplineDI->SoftSourceActor.Get())
+	{
+		PrevSplineActor->OnDestroyed.RemoveAll(this);
+	}
+	
 	SplineDI->SoftSourceActor = nullptr;
 
-	if (USplineComponent* SplineComponent = SplineComponentWeak.Get())
-    {
-		SplineComponent->TransformUpdated.RemoveAll(this);
-
-		SplineComponentWeak.Reset();
-    }
-
-	SplineComponentWeak = nullptr;
-
 	// bind
 	if (AActor* SplineActor = SplineActorWeak.Get())
 	{
 		if (USplineComponent* SplineComponent = SplineActor->FindComponentByClass<USplineComponent>())
 		{
 			SplineDI->SoftSourceActor = SplineActor;
-			SplineComponentWeak = SplineComponent;
-
-			SplineComponent->TransformUpdated.RemoveAll(this);
-			SplineComponent->TransformUpdated.AddUObject(this, &UCEClonerSplineLayout::OnSampleSplineTransformed);
-
-			USceneComponent::MarkRenderStateDirtyEvent.RemoveAll(this);
-			USceneComponent::MarkRenderStateDirtyEvent.AddUObject(this, &UCEClonerSplineLayout::OnSampleSplineRenderStateUpdated);
-
-			if (AActor* ClonerActor = GetClonerActor())
-			{
-				ClonerActor->SetActorTransform(SplineActor->GetActorTransform());
-			}
+			SplineActor->OnDestroyed.AddUniqueDynamic(this, &UCEClonerSplineLayout::OnSampleSplineDestroyed);
 		}
 	}
 }
 
-void UCEClonerSplineLayout::OnSampleSplineTransformed(USceneComponent* InComponent, EUpdateTransformFlags InFlags, ETeleportType InType)
-{
-	MarkLayoutDirty();
-}
-
 void UCEClonerSplineLayout::OnSampleSplineRenderStateUpdated(UActorComponent& InComponent)
 {
 	if (SplineActorWeak.IsValid() && InComponent.GetOwner() == SplineActorWeak.Get())
@@ -192,3 +174,8 @@ void UCEClonerSplineLayout::OnSampleSplineRenderStateUpdated(UActorComponent& In
 		MarkLayoutDirty();
 	}
 }
+
+void UCEClonerSplineLayout::OnSampleSplineDestroyed(AActor* InDestroyedActor)
+{
+	MarkLayoutDirty();
+}
diff --git a/Engine/Plugins/Experimental/ClonerEffector/Source/ClonerEffector/Public/CEClonerEffectorShared.h b/Engine/Plugins/Experimental/ClonerEffector/Source/ClonerEffector/Public/CEClonerEffectorShared.h
--- a/Engine/Plugins/Experimental/ClonerEffector/Source/ClonerEffector/Public/CEClonerEffectorShared.h
+++ b/Engine/Plugins/Experimental/ClonerEffector/Source/ClonerEffector/Public/CEClonerEffectorShared.h
@@ -250,9 +250,9 @@ enum class ECEClonerSpawnLoopMode : uint8
 UENUM()
 enum class ECEClonerSpawnBehaviorMode : uint8
 {
-	/** Spawns instantly the number of clones needed for the layout */
+	/** Spawns instantly the number of clones needed for the layout at a specific interval */
 	Instant,
-	/** Spawns at a specific rate per second during the spawn loop duration */
+	/** Spawns clones at a specific rate per second and fills the layout */
 	Rate,
 };
 
diff --git a/Engine/Plugins/Experimental/ClonerEffector/Source/ClonerEffector/Public/Cloner/Extensions/CEClonerEmitterSpawnExtension.h b/Engine/Plugins/Experimental/ClonerEffector/Source/ClonerEffector/Public/Cloner/Extensions/CEClonerEmitterSpawnExtension.h
--- a/Engine/Plugins/Experimental/ClonerEffector/Source/ClonerEffector/Public/Cloner/Extensions/CEClonerEmitterSpawnExtension.h
+++ b/Engine/Plugins/Experimental/ClonerEffector/Source/ClonerEffector/Public/Cloner/Extensions/CEClonerEmitterSpawnExtension.h
@@ -14,7 +14,6 @@ UCLASS(MinimalAPI, BlueprintType, Within=CEClonerComponent, meta=(Section="Emiss
 class UCEClonerEmitterSpawnExtension : public UCEClonerExtensionBase
 {
 	GENERATED_BODY()
-
 public:
 	UCEClonerEmitterSpawnExtension()
 		: UCEClonerExtensionBase(
@@ -68,6 +67,17 @@ public:
 		return SpawnRate;
 	}
 
+#if WITH_EDITOR
+	UFUNCTION(BlueprintCallable, Category="Cloner")
+	CLONEREFFECTOR_API void SetUseLocalSpace(bool bInLocalSpace);
+
+	UFUNCTION(BlueprintPure, Category="Cloner")
+	bool GetUseLocalSpace() const
+	{
+		return bUseLocalSpace;
+	}
+#endif
+
 protected:
 	//~ Begin UObject
 #if WITH_EDITOR
@@ -87,21 +97,27 @@ protected:
 	UPROPERTY(EditInstanceOnly, Setter, Getter, DisplayName="Emission Count", Category="Spawn", meta=(ClampMin="1", EditCondition="SpawnLoopMode == ECEClonerSpawnLoopMode::Multiple", EditConditionHides))
 	int32 SpawnLoopIterations = 1;
 
-	/** Interval/Duration of spawn for clones */
-	UPROPERTY(EditInstanceOnly, Setter, Getter, DisplayName="Emission Interval", Category="Spawn", meta=(ClampMin="0", EditCondition="SpawnLoopMode != ECEClonerSpawnLoopMode::Once", EditConditionHides))
-	float SpawnLoopInterval = 1.f;
-
 	/** How does spawn occurs */
 	UPROPERTY(EditInstanceOnly, Setter, Getter, DisplayName="Emission Style", Category="Spawn", meta=(EditCondition="SpawnLoopMode != ECEClonerSpawnLoopMode::Once", EditConditionHides))
 	ECEClonerSpawnBehaviorMode SpawnBehaviorMode = ECEClonerSpawnBehaviorMode::Instant;
 
+	/** Interval/Duration of spawn for clones */
+	UPROPERTY(EditInstanceOnly, Setter, Getter, DisplayName="Emission Interval", Category="Spawn", meta=(ClampMin="0", EditCondition="SpawnLoopMode != ECEClonerSpawnLoopMode::Once && SpawnBehaviorMode == ECEClonerSpawnBehaviorMode::Instant", EditConditionHides))
+	float SpawnLoopInterval = 1.f;
+	
 	/** How many clones to spawn each seconds */
 	UPROPERTY(EditInstanceOnly, Setter, Getter, DisplayName="Emission Rate", Category="Spawn", meta=(ClampMin="0", EditCondition="SpawnLoopMode != ECEClonerSpawnLoopMode::Once && SpawnBehaviorMode == ECEClonerSpawnBehaviorMode::Rate", EditConditionHides))
 	float SpawnRate = 1.f;
 
+	/** Simulation takes place in local space or world space */
+	UPROPERTY(EditInstanceOnly, Category="Spawn", meta=(EditCondition="SpawnLoopMode != ECEClonerSpawnLoopMode::Once", EditConditionHides))
+	bool bUseLocalSpace = true;
+
 private:
 #if WITH_EDITOR
 	/** Used for PECP */
 	static const TCEPropertyChangeDispatcher<UCEClonerEmitterSpawnExtension> PropertyChangeDispatcher;
+
+	void OnLocalSpaceChanged();
 #endif
 };
\ No newline at end of file
diff --git a/Engine/Plugins/Experimental/ClonerEffector/Source/ClonerEffector/Public/Cloner/Layouts/CEClonerLayoutBase.h b/Engine/Plugins/Experimental/ClonerEffector/Source/ClonerEffector/Public/Cloner/Layouts/CEClonerLayoutBase.h
--- a/Engine/Plugins/Experimental/ClonerEffector/Source/ClonerEffector/Public/Cloner/Layouts/CEClonerLayoutBase.h
+++ b/Engine/Plugins/Experimental/ClonerEffector/Source/ClonerEffector/Public/Cloner/Layouts/CEClonerLayoutBase.h
@@ -30,6 +30,8 @@ enum class ECEClonerLayoutAssetVersion : uint8
 	TextureConstraint,
 	/** circle ring added */
 	ConcentricCircle,
+	/** Local/world space switch added */
+	LocalWorldSpaceSwitch,
 	LatestVersionPlusOne,
 	LatestVersion = LatestVersionPlusOne - 1
 };
diff --git a/Engine/Plugins/Experimental/ClonerEffector/Source/ClonerEffector/Public/Cloner/Layouts/CEClonerMeshLayout.h b/Engine/Plugins/Experimental/ClonerEffector/Source/ClonerEffector/Public/Cloner/Layouts/CEClonerMeshLayout.h
--- a/Engine/Plugins/Experimental/ClonerEffector/Source/ClonerEffector/Public/Cloner/Layouts/CEClonerMeshLayout.h
+++ b/Engine/Plugins/Experimental/ClonerEffector/Source/ClonerEffector/Public/Cloner/Layouts/CEClonerMeshLayout.h
@@ -83,12 +83,9 @@ protected:
 	//~ End UObject
 
 	//~ Begin UCEClonerLayoutBase
-	virtual void OnLayoutInactive() override;
 	virtual void OnLayoutParametersChanged(UCEClonerComponent* InComponent) override;
 	//~ End UCEClonerLayoutBase
-
-	void OnSampleMeshTransformed(USceneComponent* InComponent, EUpdateTransformFlags InFlags, ETeleportType InType);
-
+	
 	UFUNCTION()
 	void OnSampleActorDestroyed(AActor* InDestroyedActor);
 
diff --git a/Engine/Plugins/Experimental/ClonerEffector/Source/ClonerEffector/Public/Cloner/Layouts/CEClonerSplineLayout.h b/Engine/Plugins/Experimental/ClonerEffector/Source/ClonerEffector/Public/Cloner/Layouts/CEClonerSplineLayout.h
--- a/Engine/Plugins/Experimental/ClonerEffector/Source/ClonerEffector/Public/Cloner/Layouts/CEClonerSplineLayout.h
+++ b/Engine/Plugins/Experimental/ClonerEffector/Source/ClonerEffector/Public/Cloner/Layouts/CEClonerSplineLayout.h
@@ -78,11 +78,11 @@ protected:
 	//~ End UObject
 
 	//~ Begin UCEClonerLayoutBase
+	virtual void OnLayoutActive() override;
 	virtual void OnLayoutInactive() override;
 	virtual void OnLayoutParametersChanged(UCEClonerComponent* InComponent) override;
 	//~ End UCEClonerLayoutBase
 
-	void OnSampleSplineTransformed(USceneComponent* InComponent, EUpdateTransformFlags InFlags, ETeleportType InType);
 	void OnSampleSplineRenderStateUpdated(UActorComponent& InComponent);
 
 	UPROPERTY(EditInstanceOnly, Setter, Getter, Category="Layout")
@@ -95,9 +95,9 @@ protected:
     bool bOrientMesh = false;
 
 private:
-	UPROPERTY(Transient)
-	TWeakObjectPtr<USplineComponent> SplineComponentWeak;
-
+	UFUNCTION()
+	void OnSampleSplineDestroyed(AActor* InDestroyedActor);
+	
 #if WITH_EDITOR
 	/** Used for PECP */
 	static const TCEPropertyChangeDispatcher<UCEClonerSplineLayout> PropertyChangeDispatcher;
