MotionDesign : ClonerEffector
- Refactoring the cache versioning system to automate niagara system change detection and discard old data

diff --git a/Engine/Plugins/Experimental/ClonerEffector/Source/ClonerEffector/Private/Cloner/Layouts/CEClonerLayoutBase.cpp b/Engine/Plugins/Experimental/ClonerEffector/Source/ClonerEffector/Private/Cloner/Layouts/CEClonerLayoutBase.cpp
--- a/Engine/Plugins/Experimental/ClonerEffector/Source/ClonerEffector/Private/Cloner/Layouts/CEClonerLayoutBase.cpp
+++ b/Engine/Plugins/Experimental/ClonerEffector/Source/ClonerEffector/Private/Cloner/Layouts/CEClonerLayoutBase.cpp
@@ -63,6 +63,20 @@ bool UCEClonerLayoutBase::IsLayoutValid() const
 	{
 		UE_LOG(LogCEClonerLayoutBase, Warning, TEXT("Cloner layout %s : Template system (%s) is not based off base system (%s)"), *LayoutName.ToString(), *LayoutAssetPath, LayoutBaseAssetPath);
 	}
+#if WITH_EDITOR
+	else
+	{
+		const FString LayoutSystemHash = GetLayoutHash();
+
+		if (LayoutSystemHash.IsEmpty())
+		{
+			UE_LOG(LogCEClonerLayoutBase, Warning, TEXT("Cloner layout %s : Template system (%s) hash could not be calculated"), *LayoutName.ToString(), *LayoutAssetPath);
+			return false;
+		}
+
+		UE_LOG(LogCEClonerLayoutBase, Verbose, TEXT("Cloner layout %s : Template system (%s) hash is %s"), *LayoutName.ToString(), *LayoutAssetPath, *LayoutSystemHash);
+	}
+#endif
 
 	return bIsSystemBasedOnBaseAsset;
 }
@@ -88,16 +102,21 @@ void UCEClonerLayoutBase::LoadLayout()
 	// System already cached and available with matching version
 	if (NiagaraSystem)
 	{
-		if (CachedVersion == ECEClonerLayoutAssetVersion::LatestVersion)
+		if (IsSystemHashMatching())
 		{
-			UE_LOG(LogCEClonerLayoutBase, Verbose, TEXT("%s : Cloner layout %s using cached system V%i"), *GetClonerActor()->GetActorNameOrLabel(), *LayoutName.ToString(), CachedVersion)
+			UE_LOG(LogCEClonerLayoutBase, Verbose, TEXT("%s : Cloner layout %s using cached system %s"), *GetClonerActor()->GetActorNameOrLabel(), *LayoutName.ToString(), *CachedSystemHash)
 
 			CacheMeshRenderer();
 			OnSystemLoaded();
 			return;
 		}
+		
+		FString LayoutHash;
+#if WITH_EDITOR
+		LayoutHash = GetLayoutHash();
+#endif
 
-		UE_LOG(LogCEClonerLayoutBase, Warning, TEXT("%s : Cloner layout %s skipping cached system V%i due to version mismatch V%i"), *GetClonerActor()->GetActorNameOrLabel(), *LayoutName.ToString(), CachedVersion, ECEClonerLayoutAssetVersion::LatestVersion)
+		UE_LOG(LogCEClonerLayoutBase, Warning, TEXT("%s : Cloner layout %s skipping cached system %s due to hash mismatch %s"), *GetClonerActor()->GetActorNameOrLabel(), *LayoutName.ToString(), *CachedSystemHash, *LayoutHash)
 
 		NiagaraSystem->MarkAsGarbage();
 		NiagaraSystem = nullptr;
@@ -309,7 +328,7 @@ void UCEClonerLayoutBase::PostLoad()
 {
 	Super::PostLoad();
 
-	if (CachedVersion == ECEClonerLayoutAssetVersion::PreVersioning)
+	if (CachedSystemHash.IsEmpty())
 	{
 		// After cloner layout load, niagara system should not be loaded since property was transient pre versioning,
 		// so look for it in outer chain otherwise it will trigger a world GC leak when switching level
@@ -345,7 +364,9 @@ void UCEClonerLayoutBase::OnSystemPackageLoaded(const FName& InName, UPackage* I
 		constexpr ERenameFlags RenameFlags = REN_NonTransactional | REN_DontCreateRedirectors;
 		if (NiagaraSystem->Rename(nullptr, this, RenameFlags))
 		{
-			CachedVersion = ECEClonerLayoutAssetVersion::LatestVersion;
+#if WITH_EDITOR
+			CachedSystemHash = GetLayoutHash();
+#endif
 			CacheMeshRenderer();
 		}
 
@@ -469,6 +490,32 @@ void UCEClonerLayoutBase::CleanOwnedSystem() const
 	}
 }
 
+#if WITH_EDITOR
+FString UCEClonerLayoutBase::GetLayoutHash() const
+{
+	FString LayoutHash;
+	
+	if (const UNiagaraSystem* TemplateNiagaraSystem = LoadObject<UNiagaraSystem>(nullptr, *LayoutAssetPath))
+	{
+		if (UPackage* Package = TemplateNiagaraSystem->GetPackage())
+		{ 
+			BytesToHex(Package->GetSavedHash().GetBytes(), sizeof(FIoHash::ByteArray), LayoutHash);
+		}
+	}
+	
+	return LayoutHash;
+}
+#endif
+
+bool UCEClonerLayoutBase::IsSystemHashMatching() const
+{
+	return !CachedSystemHash.IsEmpty()
+#if WITH_EDITOR
+		&& GetLayoutHash().Equals(CachedSystemHash)
+#endif
+	;
+}
+
 UCEClonerComponent* UCEClonerLayoutBase::GetClonerComponent() const
 {
 	return GetTypedOuter<UCEClonerComponent>();
diff --git a/Engine/Plugins/Experimental/ClonerEffector/Source/ClonerEffector/Public/Cloner/Layouts/CEClonerLayoutBase.h b/Engine/Plugins/Experimental/ClonerEffector/Source/ClonerEffector/Public/Cloner/Layouts/CEClonerLayoutBase.h
--- a/Engine/Plugins/Experimental/ClonerEffector/Source/ClonerEffector/Public/Cloner/Layouts/CEClonerLayoutBase.h
+++ b/Engine/Plugins/Experimental/ClonerEffector/Source/ClonerEffector/Public/Cloner/Layouts/CEClonerLayoutBase.h
@@ -12,30 +12,6 @@ class UCEClonerExtensionBase;
 class UNiagaraMeshRendererProperties;
 class UNiagaraSystem;
 
-UENUM()
-enum class ECEClonerLayoutAssetVersion : uint8
-{
-	PreVersioning = 0,
-	/** Niagara System is now cached to improve loading time */
-	CachedVersion,
-	/** Cull clones added */
-	CullEffector,
-	/** Delay clones added */
-	DelayEffector,
-	/** Cancel clones effect added */
-	CancelEffector,
-	/** Procedural pattern effect added */
-	ProceduralEffector,
-	/** Texture constraint improved */
-	TextureConstraint,
-	/** circle ring added */
-	ConcentricCircle,
-	/** Local/world space switch added */
-	LocalWorldSpaceSwitch,
-	LatestVersionPlusOne,
-	LatestVersion = LatestVersionPlusOne - 1
-};
-
 /**
  * Base class for layouts available in the cloner actor
  * Steps to add a new layout :
@@ -182,6 +158,14 @@ private:
 	/** Marks niagara system owned by this layout garbage to avoid GC leak */
 	void CleanOwnedSystem() const;
 
+#if WITH_EDITOR
+	/** Get the niagara system package hash */
+	FString GetLayoutHash() const;
+#endif
+
+	/** Checks if cached hash matches package hash */
+	bool IsSystemHashMatching() const;
+
 	/** Layout name to display in layout options */
 	UPROPERTY(Transient)
 	FName LayoutName;
@@ -198,9 +182,9 @@ private:
 	UPROPERTY(Transient, DuplicateTransient, TextExportTransient)
 	TObjectPtr<UNiagaraMeshRendererProperties> MeshRenderer;
 
-	/** Version of the cached system for diffs */
+	/** Hash of the cached system for diffs */
 	UPROPERTY(DuplicateTransient, TextExportTransient)
-	ECEClonerLayoutAssetVersion CachedVersion = ECEClonerLayoutAssetVersion::PreVersioning;
+	FString CachedSystemHash;
 
 	/** Id for the load request initiated */
 	int32 LoadRequestIdentifier = INDEX_NONE;
diff --git a/Engine/Plugins/Experimental/ClonerEffector/Source/ClonerEffectorMeshBuilder/Private/CEMeshBuilder.cpp b/Engine/Plugins/Experimental/ClonerEffector/Source/ClonerEffectorMeshBuilder/Private/CEMeshBuilder.cpp
--- a/Engine/Plugins/Experimental/ClonerEffector/Source/ClonerEffectorMeshBuilder/Private/CEMeshBuilder.cpp
+++ b/Engine/Plugins/Experimental/ClonerEffector/Source/ClonerEffectorMeshBuilder/Private/CEMeshBuilder.cpp
@@ -28,6 +28,7 @@
 #include "ProceduralMeshComponent.h"
 #include "Rendering/SkeletalMeshRenderData.h"
 #include "StaticMeshDescription.h"
+#include "StaticMeshOperations.h"
 #include "UDynamicMesh.h"
 
 const FCEMeshBuilder::FCEMeshBuilderParams FCEMeshBuilder::DefaultBuildParams;
@@ -1095,6 +1096,9 @@ bool FCEMeshBuilder::DynamicMeshToStaticMesh(UDynamicMesh* InMesh, UStaticMesh*
 	const TArray<const FMeshDescription*> MeshDescriptions {&MeshDescription};
 	bResult = OutMesh->BuildFromMeshDescriptions(MeshDescriptions, Params);
 
+	FStaticMeshOperations::ComputeTriangleTangentsAndNormals(MeshDescription);
+	FStaticMeshOperations::ComputeTangentsAndNormals(MeshDescription, EComputeNTBsFlags::Normals);
+	
 #if WITH_EDITOR
 	UStaticMesh::FBuildParameters BuildParameters;
 	BuildParameters.bInSilent = true;
