MotionDesign : ClonerEffector
- Add FreePlacement layout to mirror attached actors one to one and use all benefits of cloner and effector, when this layout is active, attached actors can be selected more easily in the editor to move them around
- Fix cloner runtime usage when level is streamed in and cached resources are not yet loaded causing it to not appear

#virtualized

diff --git a/Engine/Plugins/Experimental/ClonerEffector/Source/ClonerEffector/Private/CEClonerEffectorShared.cpp b/Engine/Plugins/Experimental/ClonerEffector/Source/ClonerEffector/Private/CEClonerEffectorShared.cpp
--- a/Engine/Plugins/Experimental/ClonerEffector/Source/ClonerEffector/Private/CEClonerEffectorShared.cpp
+++ b/Engine/Plugins/Experimental/ClonerEffector/Source/ClonerEffector/Private/CEClonerEffectorShared.cpp
@@ -23,6 +23,44 @@
 
 #define LOCTEXT_NAMESPACE "CEClonerEffectorShared"
 
+#if WITH_EDITOR
+bool FCEClonerAttachmentItem::CheckBoundsChanged(bool bInUpdate)
+{
+	bool bChanged = false;
+
+	if (const AActor* Actor = ItemActor.Get())
+	{
+		FBox LocalBox(EForceInit::ForceInit);
+
+		TArray<UPrimitiveComponent*> PrimitiveComponents;
+		Actor->GetComponents(PrimitiveComponents);
+		for (const UPrimitiveComponent* Component : PrimitiveComponents)
+		{
+			if (Component && Component->IsRegistered())
+			{
+				FTransform ComponentToActorTransform = Component->GetComponentTransform().GetRelativeTransform(Actor->GetActorTransform());
+				FBoxSphereBounds ComponentBounds = Component->CalcBounds(ComponentToActorTransform);
+
+				LocalBox += ComponentBounds.GetBox();
+			}
+		}
+
+		const FVector ActorOrigin = LocalBox.GetCenter();
+		const FVector ActorExtent = LocalBox.GetExtent();
+
+		bChanged = !ActorOrigin.Equals(Origin) || !ActorExtent.Equals(Extent);
+
+		if (bInUpdate)
+		{
+			Origin = ActorOrigin;
+			Extent = ActorExtent;
+		}
+	}
+
+	return bChanged;
+}
+#endif
+
 void FCEClonerEffectorChannelData::Write(UNiagaraDataChannelWriter* InWriter) const
 {
 	if (!InWriter)
diff --git a/Engine/Plugins/Experimental/ClonerEffector/Source/ClonerEffector/Private/Cloner/CEClonerComponent.cpp b/Engine/Plugins/Experimental/ClonerEffector/Source/ClonerEffector/Private/Cloner/CEClonerComponent.cpp
--- a/Engine/Plugins/Experimental/ClonerEffector/Source/ClonerEffector/Private/Cloner/CEClonerComponent.cpp
+++ b/Engine/Plugins/Experimental/ClonerEffector/Source/ClonerEffector/Private/Cloner/CEClonerComponent.cpp
@@ -9,8 +9,10 @@
 #include "Components/BillboardComponent.h"
 #include "Components/DynamicMeshComponent.h"
 #include "Components/SkeletalMeshComponent.h"
+#include "Engine/Level.h"
 #include "Engine/StaticMeshActor.h"
 #include "Engine/Texture2D.h"
+#include "Engine/World.h"
 #include "Materials/Material.h"
 #include "Materials/MaterialInstanceDynamic.h"
 #include "NiagaraMeshRendererProperties.h"
@@ -33,6 +35,8 @@
 UCEClonerComponent::FOnClonerInitialized UCEClonerComponent::OnClonerInitializedDelegate;
 UCEClonerComponent::FOnClonerLayoutLoaded UCEClonerComponent::OnClonerLayoutLoadedDelegate;
 UCEClonerComponent::FOnClonerMeshUpdated UCEClonerComponent::OnClonerMeshUpdatedDelegate;
+UCEClonerComponent::FOnClonerAttachmentChanged UCEClonerComponent::OnClonerActorAttachedDelegate;
+UCEClonerComponent::FOnClonerAttachmentChanged UCEClonerComponent::OnClonerActorDetachedDelegate;
 
 DEFINE_LOG_CATEGORY_STATIC(LogCEClonerComponent, Log, All);
 
@@ -185,6 +189,11 @@ void UCEClonerComponent::OnComponentDestroyed(bool bInDestroyingHierarchy)
 	
 	UMaterial::OnMaterialCompilationFinished().RemoveAll(this);
 #endif
+
+	FWorldDelegates::LevelAddedToWorld.RemoveAll(this);
+
+	FTSTicker::GetCoreTicker().RemoveTicker(ClonerTickerHandle);
+	ClonerTickerHandle.Reset();
 }
 
 void UCEClonerComponent::UpdateClonerRenderState()
@@ -251,7 +260,7 @@ void UCEClonerComponent::UpdateAttachmentTree()
 	{
 		AActor* RootChild = RootChildren[RootIdx];
 
-		if (!RootChild)
+		if (!IsValid(RootChild))
 		{
 			continue;
 		}
@@ -261,7 +270,7 @@ void UCEClonerComponent::UpdateAttachmentTree()
 		// Lets find the old root idx
 		const int32 OldIdx = ClonerTree.RootActors.Find(RootChild);
 		TObjectPtr<UStaticMesh> CombineBakedMesh = nullptr;
-		if (OldIdx != INDEX_NONE)
+		if (ClonerTree.MergedBakedMeshes.IsValidIndex(OldIdx))
 		{
 			CombineBakedMesh = ClonerTree.MergedBakedMeshes[OldIdx].Get();
 
@@ -301,6 +310,8 @@ void UCEClonerComponent::UpdateAttachmentTree()
 			UnbindActorDelegates(ClonedActor);
 			ClonerTree.ItemAttachmentMap.Remove(ClonedItem->ItemActor);
 			SetActorVisibility(ClonedActor, true);
+
+			OnClonerActorDetachedDelegate.Broadcast(this, ClonedActor);
 		}
 		else if (ClonedItem->Status == ECEClonerAttachmentStatus::Outdated)
 		{
@@ -332,7 +343,7 @@ void UCEClonerComponent::UpdateAttachmentTree()
 
 void UCEClonerComponent::UpdateActorAttachment(AActor* InActor, AActor* InParent)
 {
-	if (!InActor)
+	if (!IsValid(InActor))
 	{
 		return;
 	}
@@ -379,6 +390,15 @@ void UCEClonerComponent::UpdateActorAttachment(AActor* InActor, AActor* InParent
 			AttachmentItem->ActorTransform = ActorTransform;
 			AttachmentItem->Status = ECEClonerAttachmentStatus::Outdated;
 		}
+
+#if WITH_EDITOR
+		// Check bounds change
+		if (AttachmentItem->CheckBoundsChanged(/** CheckAndUpdate */true))
+		{
+			InvalidateBakedStaticMesh(InActor);
+			AttachmentItem->Status = ECEClonerAttachmentStatus::Outdated;
+		}
+#endif
 	}
 	else
 	{
@@ -389,9 +409,22 @@ void UCEClonerComponent::UpdateActorAttachment(AActor* InActor, AActor* InParent
 		AttachmentItem->MeshStatus = ECEClonerAttachmentStatus::Outdated;
 		AttachmentItem->bRootItem = InParent == nullptr;
 		AttachmentItem->Status = ECEClonerAttachmentStatus::Outdated;
+
+#if WITH_EDITOR
+		AttachmentItem->CheckBoundsChanged(/** CheckAndUpdate */true);
+#endif
+
 		InvalidateBakedStaticMesh(InActor);
+	}
+
+	if (!AttachmentItem->bSetupDone)
+	{
+		AttachmentItem->bSetupDone = true;
+
 		BindActorDelegates(InActor);
 		SetActorVisibility(InActor, false);
+
+		OnClonerActorAttachedDelegate.Broadcast(this, InActor);
 	}
 
 	if (AttachmentItem->ChildrenActors.Num() != ChildrenActors.Num())
@@ -409,7 +442,7 @@ void UCEClonerComponent::UpdateActorAttachment(AActor* InActor, AActor* InParent
 
 void UCEClonerComponent::BindActorDelegates(AActor* InActor)
 {
-	if (!InActor)
+	if (!IsValid(InActor))
 	{
 		return;
 	}
@@ -486,7 +519,7 @@ void UCEClonerComponent::UnbindActorDelegates(AActor* InActor) const
 	}
 }
 
-void UCEClonerComponent::SetActorVisibility(AActor* InActor, bool bInVisibility)
+void UCEClonerComponent::SetActorVisibility(AActor* InActor, bool bInVisibility, ECEClonerActorVisibility InTarget)
 {
 	if (!InActor)
 	{
@@ -494,9 +527,16 @@ void UCEClonerComponent::SetActorVisibility(AActor* InActor, bool bInVisibility)
 	}
 
 #if WITH_EDITOR
-	InActor->SetIsTemporarilyHiddenInEditor(!bInVisibility);
+	if (EnumHasAnyFlags(InTarget, ECEClonerActorVisibility::Editor))
+	{
+		InActor->SetIsTemporarilyHiddenInEditor(!bInVisibility);
+	}
 #endif
-	InActor->SetActorHiddenInGame(!bInVisibility);
+
+	if (EnumHasAnyFlags(InTarget, ECEClonerActorVisibility::Game))
+	{
+		InActor->SetActorHiddenInGame(!bInVisibility);
+	}
 }
 
 void UCEClonerComponent::OnActorDestroyed(AActor* InDestroyedActor)
@@ -508,6 +548,8 @@ void UCEClonerComponent::OnActorDestroyed(AActor* InDestroyedActor)
 		ClonerTree.ItemAttachmentMap.Remove(InDestroyedActor);
 		SetActorVisibility(InDestroyedActor, true);
 		bClonerMeshesDirty = true;
+
+		OnClonerActorDetachedDelegate.Broadcast(this, InDestroyedActor);
 	}
 }
 
@@ -627,14 +669,18 @@ void UCEClonerComponent::OnMeshChanged(UStaticMeshComponent*, AActor* InActor)
 		return;
 	}
 
-	if (FCEClonerAttachmentItem* Item = ClonerTree.ItemAttachmentMap.Find(InActor))
-	{
-		UE_LOG(LogCEClonerComponent, Log, TEXT("%s : Detected mesh change for %s"), *ClonerActor->GetActorNameOrLabel(), *InActor->GetActorNameOrLabel());
+	FCEClonerAttachmentItem* Item = ClonerTree.ItemAttachmentMap.Find(InActor);
 
-		Item->MeshStatus = ECEClonerAttachmentStatus::Outdated;
-        InvalidateBakedStaticMesh(InActor);
-		ClonerTree.DirtyItemAttachments.Add(Item->ItemActor);
+	if (!Item || Item->MeshStatus == ECEClonerAttachmentStatus::Outdated)
+	{
+		return;
 	}
+
+	UE_LOG(LogCEClonerComponent, Log, TEXT("%s : Detected mesh change for %s"), *ClonerActor->GetActorNameOrLabel(), *InActor->GetActorNameOrLabel());
+
+	Item->MeshStatus = ECEClonerAttachmentStatus::Outdated;
+	InvalidateBakedStaticMesh(InActor);
+	ClonerTree.DirtyItemAttachments.Add(Item->ItemActor);
 }
 
 void UCEClonerComponent::GetOrderedRootActors(TArray<AActor*>& OutActors) const
@@ -661,6 +707,11 @@ void UCEClonerComponent::GetOrderedRootActors(TArray<AActor*>& OutActors) const
 	{
 		ClonerActor->GetAttachedActors(OutActors, true, false);
 	}
+
+	OutActors.RemoveAll([](const AActor* InActor)
+	{
+		return !IsValid(InActor);
+	});
 }
 
 AActor* UCEClonerComponent::GetRootActor(AActor* InActor) const
@@ -683,6 +734,7 @@ AActor* UCEClonerComponent::GetRootActor(AActor* InActor) const
 
 void UCEClonerComponent::InvalidateBakedStaticMesh(AActor* InActor)
 {
+#if WITH_EDITOR
 	if (!InActor)
 	{
 		return;
@@ -704,6 +756,7 @@ void UCEClonerComponent::InvalidateBakedStaticMesh(AActor* InActor)
 			InvalidateBakedStaticMesh(FoundItem->ParentActor.Get());
 		}
 	}
+#endif
 }
 
 void UCEClonerComponent::UpdateDirtyMeshesAsync()
@@ -1056,6 +1109,19 @@ void UCEClonerComponent::FireMaterialWarning(const AActor* InContextActor, const
 #endif
 }
 
+void UCEClonerComponent::GetPreloadDependencies(TArray<UObject*>& OutDeps)
+{
+	Super::GetPreloadDependencies(OutDeps);
+
+	for (UStaticMesh* BakedMesh : ClonerTree.MergedBakedMeshes)
+	{
+		if (BakedMesh)
+		{
+			OutDeps.Add(BakedMesh);
+		}
+	}
+}
+
 #if WITH_EDITOR
 void UCEClonerComponent::ShowMaterialWarning(int32 InMaterialCount)
 {
@@ -1088,7 +1154,7 @@ FName UCEClonerComponent::GetLayoutNamePropertyName()
 
 void UCEClonerComponent::InitializeCloner()
 {
-	if (bClonerInitialized)
+	if (IsTemplate() || bClonerInitialized)
 	{
 		return;
 	}
@@ -1109,8 +1175,46 @@ void UCEClonerComponent::InitializeCloner()
 	}
 #endif
 
+	const UWorld* World = GetWorld();
+	const ULevel* Level = GetComponentLevel();
+
+	// When level is streamed in, wait until actor hierarchy and resources are ready before initializing
+	if (IsValid(World) && IsValid(Level))
+	{
+		if (!Level->IsPersistentLevel() && Level->HasAnyInternalFlags(EInternalObjectFlags_AsyncLoading))
+		{
+			FWorldDelegates::LevelAddedToWorld.AddUObject(this, &UCEClonerComponent::OnLevelLoaded);
+		}
+		else
+		{
+			PostInitializeCloner();
+		}
+	}
+}
+
+void UCEClonerComponent::OnLevelLoaded(ULevel* InLevel, UWorld* InWorld)
+{
+	const ULevel* Level = GetComponentLevel();
+	const UWorld* World = GetWorld();
+
+	if (!IsValid(Level)
+		|| !IsValid(World)
+		|| Level != InLevel
+		|| Level->HasAnyInternalFlags(EInternalObjectFlags_AsyncLoading))
+	{
+		return;
+	}
+
+	FWorldDelegates::LevelAddedToWorld.RemoveAll(this);
+
+	PostInitializeCloner();
+}
+
+void UCEClonerComponent::PostInitializeCloner()
+{
 	// Register a custom ticker to avoid using the component tick that needs the simulation to be solo
 	TreeUpdateDeltaTime = TreeUpdateInterval;
+
 	RegisterTicker();
 
 	// Load layout after registering ticker to let attachment tree update, then layout rendering will occur with up to date attachment data
@@ -1121,6 +1225,12 @@ void UCEClonerComponent::InitializeCloner()
 
 void UCEClonerComponent::RegisterTicker()
 {
+	if (!bClonerInitialized)
+	{
+		return;
+	}
+
+	// Register custom ticker to avoid using component tick and niagara solo mode
 	FTSTicker::GetCoreTicker().RemoveTicker(ClonerTickerHandle);
 	ClonerTickerHandle = FTSTicker::GetCoreTicker().AddTicker(FTickerDelegate::CreateUObject(this, &UCEClonerComponent::TickCloner));
 }
@@ -1140,7 +1250,6 @@ bool UCEClonerComponent::TickCloner(float InDelta)
 		if (TreeUpdateDeltaTime >= TreeUpdateInterval)
 		{
 			TreeUpdateDeltaTime -= TreeUpdateInterval != 0.f ? TreeUpdateInterval : TreeUpdateDeltaTime;
-
 			UpdateClonerAttachmentTree();
 			UpdateClonerRenderState();
 		}
@@ -1345,7 +1454,7 @@ void UCEClonerComponent::CreateDefaultActorAttached()
 
 	// Spawn attached actor with same flags as this actor
 	FActorSpawnParameters SpawnParameters;
-	SpawnParameters.Owner = Owner;
+	SpawnParameters.OverrideLevel = Owner->GetLevel();
 	SpawnParameters.ObjectFlags = GetFlags() | RF_Transactional;
 	SpawnParameters.bTemporaryEditorActor = false;
 
@@ -1600,7 +1709,6 @@ void UCEClonerComponent::OnVisualizerSpriteVisibleChanged()
 		}
 	}
 }
-#endif
 
 void UCEClonerComponent::OnRenderStateDirty(UActorComponent& InActorComponent)
 {
@@ -1620,7 +1728,7 @@ void UCEClonerComponent::OnRenderStateDirty(UActorComponent& InActorComponent)
 
 	FCEClonerAttachmentItem* Item = ClonerTree.ItemAttachmentMap.Find(Owner);
 
-	if (!Item)
+	if (!Item || !Item->CheckBoundsChanged(/** CheckAndUpdate */false))
 	{
 		return;
 	}
@@ -1634,6 +1742,7 @@ void UCEClonerComponent::OnRenderStateDirty(UActorComponent& InActorComponent)
 	InvalidateBakedStaticMesh(Owner);
 	ClonerTree.DirtyItemAttachments.Add(Item->ItemActor);
 }
+#endif
 
 void UCEClonerComponent::OnComponentTransformed(USceneComponent* InComponent, EUpdateTransformFlags InFlags, ETeleportType InTeleport)
 {
@@ -1841,6 +1950,16 @@ UCEClonerExtensionBase* UCEClonerComponent::GetExtension(FName InExtensionName)
 	return nullptr;
 }
 
+TArray<AActor*> UCEClonerComponent::GetClonerRootActors() const
+{
+	TArray<AActor*> RootActors;
+	Algo::Transform(ClonerTree.RootActors, RootActors, [](const TWeakObjectPtr<AActor>& InActorWeak)->AActor*
+	{
+		return InActorWeak.Get();
+	});
+	return RootActors;
+}
+
 void UCEClonerComponent::OnActiveLayoutLoaded(UCEClonerLayoutBase* InLayout, bool bInSuccess)
 {
 	if (!InLayout)
@@ -1968,6 +2087,8 @@ void UCEClonerComponent::UpdateClonerMeshes()
 		return;
 	}
 
+	bool bMeshChanged = MeshRenderer->Meshes.Num() != ClonerTree.MergedBakedMeshes.Num();
+
 	if (bClonerMeshesDirty)
 	{
 		// Resize mesh array properly
@@ -1981,7 +2102,14 @@ void UCEClonerComponent::UpdateClonerMeshes()
 		{
 			UStaticMesh* StaticMesh = ClonerTree.MergedBakedMeshes[Idx].Get();
 
+			// Make sure static mesh is loaded and available when streaming
+			if (StaticMesh && StaticMesh->HasAnyInternalFlags(EInternalObjectFlags_AsyncLoading))
+			{
+				return;
+			}
+
 			FNiagaraMeshRendererMeshProperties& MeshProperties = !MeshRenderer->Meshes.IsValidIndex(Idx) ? MeshRenderer->Meshes.AddDefaulted_GetRef() : MeshRenderer->Meshes[Idx];
+			bMeshChanged |= MeshProperties.Mesh != StaticMesh;
 			MeshProperties.Mesh = StaticMesh && StaticMesh->GetNumTriangles(0) > 0 ? StaticMesh : nullptr;
 
 			if (ClonerTree.RootActors.IsValidIndex(Idx))
@@ -1995,8 +2123,6 @@ void UCEClonerComponent::UpdateClonerMeshes()
 		}
 
 		bClonerMeshesDirty = !ClonerTree.DirtyItemAttachments.IsEmpty();
-
-		UE_LOG(LogCEClonerComponent, Log, TEXT("%s : Cloner mesh updated %i"), *ClonerActor->GetActorNameOrLabel(), ClonerTree.MergedBakedMeshes.Num())
 	}
 
 	for (const TObjectPtr<UCEClonerExtensionBase>& ActiveExtension : ActiveExtensions)
@@ -2004,6 +2130,14 @@ void UCEClonerComponent::UpdateClonerMeshes()
 		ActiveExtension->OnClonerMeshesUpdated();
 	}
 
+	// Extensions could override mesh renderer meshes array
+	bMeshChanged |= ClonerTree.MergedBakedMeshes.Num() != MeshRenderer->Meshes.Num();
+
+	if (bMeshChanged)
+	{
+		UE_LOG(LogCEClonerComponent, Log, TEXT("%s : Cloner mesh updated - %i baked meshes - %i rendered meshes"), *ClonerActor->GetActorNameOrLabel(), ClonerTree.MergedBakedMeshes.Num(), MeshRenderer->Meshes.Num())
+	}
+
 	// Set new number of meshes in renderer
 	SetIntParameter(TEXT("MeshNum"), MeshRenderer->Meshes.Num());
 
diff --git a/Engine/Plugins/Experimental/ClonerEffector/Source/ClonerEffector/Private/Cloner/Layouts/CEClonerFreePlacementLayout.cpp b/Engine/Plugins/Experimental/ClonerEffector/Source/ClonerEffector/Private/Cloner/Layouts/CEClonerFreePlacementLayout.cpp
new file mode 100644
--- /dev/null
+++ b/Engine/Plugins/Experimental/ClonerEffector/Source/ClonerEffector/Private/Cloner/Layouts/CEClonerFreePlacementLayout.cpp
@@ -0,0 +1,216 @@
+// Copyright Epic Games, Inc. All Rights Reserved.
+
+#include "Cloner/Layouts/CEClonerFreePlacementLayout.h"
+
+#include "Cloner/CEClonerComponent.h"
+#include "NiagaraDataInterfaceArrayFloat.h"
+
+#if WITH_EDITOR
+const TCEPropertyChangeDispatcher<UCEClonerFreePlacementLayout> UCEClonerFreePlacementLayout::PropertyChangeDispatcher =
+{
+};
+
+void UCEClonerFreePlacementLayout::PostEditChangeProperty(FPropertyChangedEvent& InPropertyChangedEvent)
+{
+	Super::PostEditChangeProperty(InPropertyChangedEvent);
+
+	PropertyChangeDispatcher.OnPropertyChanged(this, InPropertyChangedEvent);
+}
+#endif
+
+void UCEClonerFreePlacementLayout::OnLayoutActive()
+{
+	Super::OnLayoutActive();
+
+	if (UCEClonerComponent* ClonerComponent = GetClonerComponent())
+	{
+		UCEClonerComponent::OnClonerMeshUpdated().AddUObject(this, &UCEClonerFreePlacementLayout::OnClonerMeshUpdated);
+
+#if WITH_EDITOR
+		UCEClonerComponent::OnClonerActorAttached().AddUObject(this, &UCEClonerFreePlacementLayout::OnClonerActorAttached);
+		UCEClonerComponent::OnClonerActorDetached().AddUObject(this, &UCEClonerFreePlacementLayout::OnClonerActorDetached);
+		
+		// Disable selection in order to select the underlying cloned actors to move them easily
+		ClonerComponent->bSelectable = false;
+
+		// Render in custom pass to select actor without rendering its geometry
+		ApplyComponentsSettings();
+#endif
+	}
+}
+
+void UCEClonerFreePlacementLayout::OnLayoutInactive()
+{
+	Super::OnLayoutInactive();
+	
+	if (UCEClonerComponent* ClonerComponent = GetClonerComponent())
+	{
+		UCEClonerComponent::OnClonerMeshUpdated().RemoveAll(this);
+		
+#if WITH_EDITOR
+		UCEClonerComponent::OnClonerActorAttached().RemoveAll(this);
+		UCEClonerComponent::OnClonerActorDetached().RemoveAll(this);
+		
+		ClonerComponent->bSelectable = true;
+
+		RestoreComponentsSettings();
+#endif
+	}
+}
+
+void UCEClonerFreePlacementLayout::OnLayoutParametersChanged(UCEClonerComponent* InComponent)
+{
+	Super::OnLayoutParametersChanged(InComponent);
+
+	const AActor* ClonerActor = GetClonerActor();
+
+	if (!ClonerActor)
+	{
+		return;
+	}
+
+	const FNiagaraUserRedirectionParameterStore& ExposedParameters = InComponent->GetOverrideParameters();
+	
+	static const FNiagaraVariable SpawnPositionsVar(FNiagaraTypeDefinition(UNiagaraDataInterfaceArrayFloat3::StaticClass()), TEXT("SpawnPositions"));
+	UNiagaraDataInterfaceArrayFloat3* PositionsDI = Cast<UNiagaraDataInterfaceArrayFloat3>(ExposedParameters.GetDataInterface(SpawnPositionsVar));
+
+	if (!PositionsDI)
+	{
+		return;
+	}
+	
+	const TArray<AActor*> RootActors = InComponent->GetClonerRootActors();
+
+	TArray<FVector3f>& Positions = PositionsDI->GetArrayReference();
+	Positions.Empty(RootActors.Num());
+
+	// Position relative to cloner actor
+	const FTransform ClonerTransform = ClonerActor->GetActorTransform();
+	for (const AActor* RootActor : RootActors)
+	{
+		FVector3f Position = FVector3f::ZeroVector;
+
+		if (RootActor)
+		{
+			Position = FVector3f(RootActor->GetActorTransform().GetRelativeTransform(ClonerTransform).GetLocation());
+		}
+		
+		Positions.Add(Position);
+	}
+}
+
+void UCEClonerFreePlacementLayout::OnClonerMeshUpdated(UCEClonerComponent* InClonerComponent)
+{
+	if (InClonerComponent->GetActiveLayout() != this)
+	{
+		return;
+	}
+	
+	MarkLayoutDirty();
+}
+
+#if WITH_EDITOR
+void UCEClonerFreePlacementLayout::OnClonerActorAttached(UCEClonerComponent* InClonerComponent, AActor* InActor)
+{
+	if (InClonerComponent->GetActiveLayout() != this)
+	{
+		return;
+	}
+
+	ApplyComponentsSettings(InActor);
+}
+
+void UCEClonerFreePlacementLayout::OnClonerActorDetached(UCEClonerComponent* InClonerComponent, AActor* InActor)
+{
+	if (InClonerComponent->GetActiveLayout() != this)
+	{
+		return;
+	}
+
+	RestoreComponentsSettings(InActor);
+}
+
+void UCEClonerFreePlacementLayout::ApplyComponentsSettings()
+{
+	const AActor* ClonerActor = GetClonerActor();
+
+	if (!IsValid(ClonerActor))
+	{
+		return;
+	}
+
+	TArray<AActor*> AttachedActors;
+	ClonerActor->GetAttachedActors(AttachedActors, /** Reset */true, /** IncludeChildren */true);
+
+	for (AActor* AttachedActor : AttachedActors)
+	{
+		ApplyComponentsSettings(AttachedActor);
+	}
+}
+
+void UCEClonerFreePlacementLayout::RestoreComponentsSettings()
+{
+	const AActor* ClonerActor = GetClonerActor();
+
+	if (!IsValid(ClonerActor))
+	{
+		return;
+	}
+
+	TArray<AActor*> AttachedActors;
+	ClonerActor->GetAttachedActors(AttachedActors, /** Reset */true, /** IncludeChildren */true);
+
+	for (AActor* AttachedActor : AttachedActors)
+	{
+		if (AttachedActor)
+		{
+			AttachedActor->SetIsTemporarilyHiddenInEditor(true);
+			RestoreComponentsSettings(AttachedActor);
+		}
+	}
+}
+
+void UCEClonerFreePlacementLayout::ApplyComponentsSettings(AActor* InActor)
+{
+	if (!IsValid(InActor))
+	{
+		return;
+	}
+	
+	InActor->SetIsTemporarilyHiddenInEditor(false);
+
+	TArray<UPrimitiveComponent*> PrimitiveComponents;
+	InActor->GetComponents(PrimitiveComponents);
+
+	for (UPrimitiveComponent* PrimitiveComponent : PrimitiveComponents)
+	{
+		if (PrimitiveComponent)
+		{
+			PrimitiveComponent->SetRenderInMainPass(false);
+			PrimitiveComponent->SetRenderInDepthPass(false);
+			PrimitiveComponent->SetRenderCustomDepth(true);
+		}
+	}
+}
+
+void UCEClonerFreePlacementLayout::RestoreComponentsSettings(AActor* InActor)
+{
+	if (!IsValid(InActor))
+	{
+		return;
+	}
+	
+	TArray<UPrimitiveComponent*> PrimitiveComponents;
+	InActor->GetComponents(PrimitiveComponents);
+
+	for (UPrimitiveComponent* PrimitiveComponent : PrimitiveComponents)
+	{
+		if (PrimitiveComponent)
+		{
+			PrimitiveComponent->SetRenderInMainPass(true);
+			PrimitiveComponent->SetRenderInDepthPass(true);
+			PrimitiveComponent->SetRenderCustomDepth(false);
+		}
+	}
+}
+#endif
\ No newline at end of file
diff --git a/Engine/Plugins/Experimental/ClonerEffector/Source/ClonerEffector/Public/CEClonerEffectorShared.h b/Engine/Plugins/Experimental/ClonerEffector/Source/ClonerEffector/Public/CEClonerEffectorShared.h
--- a/Engine/Plugins/Experimental/ClonerEffector/Source/ClonerEffector/Public/CEClonerEffectorShared.h
+++ b/Engine/Plugins/Experimental/ClonerEffector/Source/ClonerEffector/Public/CEClonerEffectorShared.h
@@ -291,6 +291,15 @@ enum class ECEClonerSystemStatus : uint8
 	SimulationDirty = 1 << 1
 };
 
+/** Enumerates all visibility options for an actor */
+enum class ECEClonerActorVisibility : uint8
+{
+	None = 0,
+	Game = 1 << 0,
+	Editor = 1 << 1,
+	All = Game | Editor
+};
+
 USTRUCT()
 struct FCEClonerAttachmentItem
 {
@@ -329,10 +338,24 @@ struct FCEClonerAttachmentItem
 	UPROPERTY()
 	bool bRootItem = false;
 
+	/** Item setup is done (visibility, delegates) once attached/loaded */
+	UPROPERTY(Transient)
+	bool bSetupDone = false;
+
 	/** Status of this attachment item */
-	UPROPERTY()
+	UPROPERTY(Transient)
 	ECEClonerAttachmentStatus Status = ECEClonerAttachmentStatus::Outdated;
 
+#if WITH_EDITORONLY_DATA
+	/** Cached pivot to detect changes in geometry */
+	UPROPERTY(Transient)
+	FVector Origin;
+
+	/** Cached extent to detect changes in geometry */
+	UPROPERTY(Transient)
+	FVector Extent;
+#endif
+
 	friend uint32 GetTypeHash(const FCEClonerAttachmentItem& InItem)
 	{
 		return GetTypeHash(InItem.ItemActor);
@@ -342,6 +365,10 @@ struct FCEClonerAttachmentItem
 	{
 		return ItemActor.Get() == InOther.ItemActor.Get();
 	}
+
+#if WITH_EDITOR
+	bool CheckBoundsChanged(bool bInUpdate);
+#endif
 };
 
 USTRUCT()
@@ -357,8 +384,10 @@ struct FCEClonerAttachmentTree
 	UPROPERTY()
 	TArray<TWeakObjectPtr<AActor>> RootActors;
 
-	/** Merged static meshes corresponding to the root actors for niagara,
-		result of merging dynamic meshes together */
+	/**
+	 * Merged static meshes corresponding to the root actors for niagara,
+	 * result of merging dynamic meshes together
+	 */
 	UPROPERTY()
 	TArray<TObjectPtr<UStaticMesh>> MergedBakedMeshes;
 
@@ -372,10 +401,10 @@ struct FCEClonerAttachmentTree
 
 	void Reset()
 	{
-		ItemAttachmentMap.Empty();
-		RootActors.Empty();
-		MergedBakedMeshes.Empty();
-		DirtyItemAttachments.Empty();
+		ItemAttachmentMap.Reset();
+		RootActors.Reset();
+		MergedBakedMeshes.Reset();
+		DirtyItemAttachments.Reset();
 		Status = ECEClonerAttachmentStatus::Updated;
 	}
 };
diff --git a/Engine/Plugins/Experimental/ClonerEffector/Source/ClonerEffector/Public/Cloner/CEClonerComponent.h b/Engine/Plugins/Experimental/ClonerEffector/Source/ClonerEffector/Public/Cloner/CEClonerComponent.h
--- a/Engine/Plugins/Experimental/ClonerEffector/Source/ClonerEffector/Public/Cloner/CEClonerComponent.h
+++ b/Engine/Plugins/Experimental/ClonerEffector/Source/ClonerEffector/Public/Cloner/CEClonerComponent.h
@@ -12,6 +12,7 @@
 
 class UCEClonerExtensionBase;
 class UCEClonerLayoutBase;
+class ULevel;
 class UMaterialInterface;
 struct FNiagaraMeshMaterialOverride;
 
@@ -29,6 +30,7 @@ class UCEClonerComponent : public UNiagaraComponent
 	DECLARE_MULTICAST_DELEGATE_OneParam(FOnClonerMeshUpdated, UCEClonerComponent* /** ClonerComponent */)
 	DECLARE_MULTICAST_DELEGATE_TwoParams(FOnClonerLayoutLoaded, UCEClonerComponent* /** ClonerComponent */, UCEClonerLayoutBase* /** InLayout */)
 	DECLARE_MULTICAST_DELEGATE_OneParam(FOnClonerInitialized, UCEClonerComponent* /** ClonerComponent */)
+	DECLARE_MULTICAST_DELEGATE_TwoParams(FOnClonerAttachmentChanged, UCEClonerComponent* /** ClonerComponent */, AActor* /** InActor */)
 
 public:
 	/** Only materials transient or part of the content folder can be dirtied, engine or plugins cannot */
@@ -37,6 +39,8 @@ public:
 	/** Check if material has niagara usage flag set */
 	static bool IsMaterialUsageFlagSet(const UMaterialInterface* InMaterial);
 
+	static void SetActorVisibility(AActor* InActor, bool bInVisibility, ECEClonerActorVisibility InTarget = ECEClonerActorVisibility::All);
+
 #if WITH_EDITOR
 	/** Show material warning notification when missing niagara usage flag */
 	static void ShowMaterialWarning(int32 InMaterialCount);
@@ -63,6 +67,16 @@ public:
 		return OnClonerInitializedDelegate;
 	}
 
+	static FOnClonerAttachmentChanged::RegistrationType& OnClonerActorAttached()
+	{
+		return OnClonerActorAttachedDelegate;
+	}
+
+	static FOnClonerAttachmentChanged::RegistrationType& OnClonerActorDetached()
+	{
+		return OnClonerActorDetachedDelegate;
+	}
+
 	UCEClonerComponent();
 
 	UFUNCTION(BlueprintCallable, Category="Cloner")
@@ -229,6 +243,9 @@ public:
 		OutExtensions = ActiveExtensions;
 	}
 
+	/** Retrieve root actors in their respective order */
+	TArray<AActor*> GetClonerRootActors() const;
+
 protected:
 	/** Called when meshes have been updated */
 	CLONEREFFECTOR_API static FOnClonerMeshUpdated OnClonerMeshUpdatedDelegate;
@@ -239,6 +256,12 @@ protected:
 	/** Called when cloner is initialized */
 	CLONEREFFECTOR_API static FOnClonerInitialized OnClonerInitializedDelegate;
 
+	/** Called when cloner actor is attached */
+	CLONEREFFECTOR_API static FOnClonerAttachmentChanged OnClonerActorAttachedDelegate;
+
+	/** Called when cloner actor is detached */
+	CLONEREFFECTOR_API static FOnClonerAttachmentChanged OnClonerActorDetachedDelegate;
+
 	/** Replaces all unsupported material by default material, gathers unset materials that needs recompiling with proper flags */
 	static bool FilterSupportedMaterials(TArray<TWeakObjectPtr<UMaterialInterface>>& InMaterials, TArray<TWeakObjectPtr<UMaterialInterface>>& OutUnsetMaterials, UMaterialInterface* InDefaultMaterial);
 
@@ -248,6 +271,7 @@ protected:
 	void FireMaterialWarning(const AActor* InContextActor, const TArray<TWeakObjectPtr<UMaterialInterface>>& InUnsetMaterials);
 
 	//~ Begin UObject
+	virtual void GetPreloadDependencies(TArray<UObject*>& OutDeps) override;
 	virtual void PostInitProperties() override;
 	virtual void PostLoad() override;
 #if WITH_EDITOR
@@ -309,7 +333,7 @@ protected:
 
 	/** Interval to update the attachment tree and update the cloner meshes, 0 means each tick */
 	UPROPERTY(EditInstanceOnly, Setter, Getter, Category="Cloner", AdvancedDisplay, meta=(ClampMin="0"))
-	float TreeUpdateInterval = 0.2f;
+	float TreeUpdateInterval = 0.f;
 
 	/** Cloner instance seed for random deterministic patterns */
 	UPROPERTY(EditInstanceOnly, Setter, Getter, Category="Cloner")
@@ -350,6 +374,8 @@ private:
 
 	/** Initiate and perform operation */
 	void InitializeCloner();
+	void OnLevelLoaded(ULevel* InLevel, UWorld* InWorld);
+	void PostInitializeCloner();
 	void RegisterTicker();
 	bool TickCloner(float InDelta);
 
@@ -378,11 +404,12 @@ private:
 	void BindActorDelegates(AActor* InActor);
 	void UnbindActorDelegates(AActor* InActor) const;
 
-	void SetActorVisibility(AActor* InActor, bool bInVisibility);
-
 #if WITH_EDITOR
 	void OnActorPropertyChanged(UObject* InObject, FPropertyChangedEvent& InPropertyChangedEvent);
 	void OnMaterialCompiled(UMaterialInterface* InMaterial);
+
+	/** Called when the render state of a cloned component changes */
+	void OnRenderStateDirty(UActorComponent& InActorComponent);
 #endif
 
 	/** Called when a cloned actor is destroyed */
@@ -395,9 +422,6 @@ private:
 	/** Called when a mesh has changed in any cloned actors */
 	void OnMeshChanged(UStaticMeshComponent*, AActor* InActor);
 
-	/** Called when the render state of a cloned component changes */
-	void OnRenderStateDirty(UActorComponent& InActorComponent);
-
 	/** Called when the transform state of a cloned component changes */
 	void OnComponentTransformed(USceneComponent* InComponent, EUpdateTransformFlags InFlags, ETeleportType InTeleport);
 
diff --git a/Engine/Plugins/Experimental/ClonerEffector/Source/ClonerEffector/Public/Cloner/Layouts/CEClonerFreePlacementLayout.h b/Engine/Plugins/Experimental/ClonerEffector/Source/ClonerEffector/Public/Cloner/Layouts/CEClonerFreePlacementLayout.h
new file mode 100644
--- /dev/null
+++ b/Engine/Plugins/Experimental/ClonerEffector/Source/ClonerEffector/Public/Cloner/Layouts/CEClonerFreePlacementLayout.h
@@ -0,0 +1,52 @@
+// Copyright Epic Games, Inc. All Rights Reserved.
+
+#pragma once
+
+#include "CEClonerEffectorShared.h"
+#include "CEClonerLayoutBase.h"
+#include "CEPropertyChangeDispatcher.h"
+#include "CEClonerFreePlacementLayout.generated.h"
+
+UCLASS(MinimalAPI, BlueprintType)
+class UCEClonerFreePlacementLayout : public UCEClonerLayoutBase
+{
+	GENERATED_BODY()
+
+public:
+	UCEClonerFreePlacementLayout()
+		: UCEClonerLayoutBase(
+			TEXT("FreePlacement")
+			, TEXT("/Script/Niagara.NiagaraSystem'/ClonerEffector/Systems/NS_ClonerFreePlacement.NS_ClonerFreePlacement'")
+		)
+	{}
+
+protected:
+	//~ Begin UObject
+#if WITH_EDITOR
+	virtual void PostEditChangeProperty(FPropertyChangedEvent& InPropertyChangedEvent) override;
+#endif
+	//~ End UObject
+
+	//~ Begin UCEClonerLayoutBase
+	virtual void OnLayoutActive() override;
+	virtual void OnLayoutInactive() override;
+	virtual void OnLayoutParametersChanged(UCEClonerComponent* InComponent) override;
+	//~ End UCEClonerLayoutBase
+	
+	void OnClonerMeshUpdated(UCEClonerComponent* InClonerComponent);
+
+#if WITH_EDITOR
+	void OnClonerActorAttached(UCEClonerComponent* InClonerComponent, AActor* InActor);
+	void OnClonerActorDetached(UCEClonerComponent* InClonerComponent, AActor* InActor);
+	void ApplyComponentsSettings();
+	void RestoreComponentsSettings();
+	void ApplyComponentsSettings(AActor* InActor);
+	void RestoreComponentsSettings(AActor* InActor);
+#endif
+	
+private:
+#if WITH_EDITOR
+	/** Used for PECP */
+	static const TCEPropertyChangeDispatcher<UCEClonerFreePlacementLayout> PropertyChangeDispatcher;
+#endif
+};
\ No newline at end of file
