MotionDesign : ClonerEffector
- Allow usage of live texture in constraint like MediaTexture that updates on particle update, instead of killing particles, it will be hidden to change its state on next update
- Moved away from using VisibilityTag to MeshIndex to hide/show particles in multiple modules to use the same approach

#virtualized

diff --git a/Engine/Plugins/Experimental/ClonerEffector/Source/ClonerEffector/Private/Cloner/Extensions/CEClonerConstraintExtension.cpp b/Engine/Plugins/Experimental/ClonerEffector/Source/ClonerEffector/Private/Cloner/Extensions/CEClonerConstraintExtension.cpp
--- a/Engine/Plugins/Experimental/ClonerEffector/Source/ClonerEffector/Private/Cloner/Extensions/CEClonerConstraintExtension.cpp
+++ b/Engine/Plugins/Experimental/ClonerEffector/Source/ClonerEffector/Private/Cloner/Extensions/CEClonerConstraintExtension.cpp
@@ -130,6 +130,7 @@ void UCEClonerConstraintExtension::SetTextureCompareMode(ECEClonerCompareMode In
 
 void UCEClonerConstraintExtension::SetTextureThreshold(float InThreshold)
 {
+	InThreshold = FMath::Max(0, InThreshold);
 	if (FMath::IsNearlyEqual(TextureThreshold, InThreshold))
 	{
 		return;
diff --git a/Engine/Plugins/Experimental/ClonerEffector/Source/ClonerEffector/Public/Cloner/Extensions/CEClonerConstraintExtension.h b/Engine/Plugins/Experimental/ClonerEffector/Source/ClonerEffector/Public/Cloner/Extensions/CEClonerConstraintExtension.h
--- a/Engine/Plugins/Experimental/ClonerEffector/Source/ClonerEffector/Public/Cloner/Extensions/CEClonerConstraintExtension.h
+++ b/Engine/Plugins/Experimental/ClonerEffector/Source/ClonerEffector/Public/Cloner/Extensions/CEClonerConstraintExtension.h
@@ -215,7 +215,7 @@ protected:
 	UPROPERTY(EditInstanceOnly, Setter, Getter, Category="Constraint", meta=(EditCondition="Constraint == ECEClonerGridConstraint::Texture", EditConditionHides))
 	ECEClonerCompareMode TextureCompareMode = ECEClonerCompareMode::Greater;
 
-	UPROPERTY(EditInstanceOnly, Setter, Getter, Category="Constraint", meta=(EditCondition="Constraint == ECEClonerGridConstraint::Texture", EditConditionHides))
+	UPROPERTY(EditInstanceOnly, Setter, Getter, Category="Constraint", meta=(ClampMin="0", EditCondition="Constraint == ECEClonerGridConstraint::Texture", EditConditionHides))
 	float TextureThreshold = 0.f;
 
 	UPROPERTY(EditInstanceOnly, Setter, Getter, Category="Constraint", meta=(EditCondition="Constraint == ECEClonerGridConstraint::Texture", EditConditionHides))
@@ -224,7 +224,7 @@ protected:
 	UPROPERTY(EditInstanceOnly, Setter, Getter, Category="Constraint", meta=(EditCondition="Constraint == ECEClonerGridConstraint::Texture", EditConditionHides))
 	float TextureRotation = 0.f;
 	
-	UPROPERTY(EditInstanceOnly, Setter, Getter, Category="Constraint", meta=(MotionDesignVectorWidget, AllowPreserveRatio="XY", EditCondition="Constraint == ECEClonerGridConstraint::Texture", EditConditionHides))
+	UPROPERTY(EditInstanceOnly, Setter, Getter, Category="Constraint", meta=(Delta="0.01", MotionDesignVectorWidget, AllowPreserveRatio="XY", EditCondition="Constraint == ECEClonerGridConstraint::Texture", EditConditionHides))
 	FVector2D TextureScale = FVector2D(1.0);
 
 	UPROPERTY(EditInstanceOnly, Setter="SetTextureClamp", Getter="GetTextureClamp", Category="Constraint", meta=(EditCondition="Constraint == ECEClonerGridConstraint::Texture", EditConditionHides))
