MotionDesign : ClonerEffector
- Added effector forces constrain and restore capability, particles will solve forces in the effector inner zone and restore back to their rest state when outside of it, instead of flying away indefinitely
- Added texture displacement for grid layout: offset, rotation scale can be adjusted along with texture offset, rotation, scale, clamp. Live texture are supported
- Moving shared code into niagara functions to be reused in other modules
- Added GlobalScale on cloner component to adjust attachments scale more easily
- Some details panel category adjustments

#virtualized

diff --git a/Engine/Plugins/Experimental/ClonerEffector/Source/ClonerEffector/Private/Cloner/CEClonerComponent.cpp b/Engine/Plugins/Experimental/ClonerEffector/Source/ClonerEffector/Private/Cloner/CEClonerComponent.cpp
--- a/Engine/Plugins/Experimental/ClonerEffector/Source/ClonerEffector/Private/Cloner/CEClonerComponent.cpp
+++ b/Engine/Plugins/Experimental/ClonerEffector/Source/ClonerEffector/Private/Cloner/CEClonerComponent.cpp
@@ -155,12 +155,15 @@ void UCEClonerComponent::PostEditUndo()
 
 const TCEPropertyChangeDispatcher<UCEClonerComponent> UCEClonerComponent::PropertyChangeDispatcher =
 {
+	/** General */
 	{ GET_MEMBER_NAME_CHECKED(UCEClonerComponent, bEnabled), &UCEClonerComponent::OnEnabledChanged },
 	{ GET_MEMBER_NAME_CHECKED(UCEClonerComponent, Seed), &UCEClonerComponent::OnSeedChanged },
 	{ GET_MEMBER_NAME_CHECKED(UCEClonerComponent, Color), &UCEClonerComponent::OnColorChanged },
+	{ GET_MEMBER_NAME_CHECKED(UCEClonerComponent, GlobalScale), &UCEClonerComponent::OnGlobalScaleChanged },
+	{ GET_MEMBER_NAME_CHECKED(UCEClonerComponent, GlobalRotation), &UCEClonerComponent::OnGlobalRotationChanged },
+	{ GET_MEMBER_NAME_CHECKED(UCEClonerComponent, bVisualizerSpriteVisible), &UCEClonerComponent::OnVisualizerSpriteVisibleChanged },
 	/** Layout */
 	{ GET_MEMBER_NAME_CHECKED(UCEClonerComponent, LayoutName), &UCEClonerComponent::OnLayoutNameChanged },
-	{ GET_MEMBER_NAME_CHECKED(UCEClonerComponent, bVisualizerSpriteVisible), &UCEClonerComponent::OnVisualizerSpriteVisibleChanged },
 };
 
 void UCEClonerComponent::PostEditChangeProperty(FPropertyChangedEvent& InPropertyChangedEvent)
@@ -1323,6 +1326,28 @@ void UCEClonerComponent::SetColor(const FLinearColor& InColor)
 	OnColorChanged();
 }
 
+void UCEClonerComponent::SetGlobalScale(const FVector& InScale)
+{
+	if (GlobalScale.Equals(InScale))
+	{
+		return;
+	}
+
+	GlobalScale = InScale.ComponentMax(FVector(UE_KINDA_SMALL_NUMBER));;
+	OnGlobalScaleChanged();
+}
+
+void UCEClonerComponent::SetGlobalRotation(const FRotator& InRotation)
+{
+	if (GlobalRotation.Equals(InRotation))
+	{
+		return;
+	}
+
+	GlobalRotation = InRotation;
+	OnGlobalRotationChanged();
+}
+
 void UCEClonerComponent::SetLayoutName(FName InLayoutName)
 {
 	if (LayoutName == InLayoutName)
@@ -1670,6 +1695,17 @@ void UCEClonerComponent::OnColorChanged()
 	SetColorParameter(TEXT("EffectorDefaultColor"), Color);
 }
 
+void UCEClonerComponent::OnGlobalScaleChanged()
+{
+	GlobalScale = GlobalScale.ComponentMax(FVector(UE_KINDA_SMALL_NUMBER));
+	bClonerMeshesDirty = true;
+}
+
+void UCEClonerComponent::OnGlobalRotationChanged()
+{
+	bClonerMeshesDirty = true;
+}
+
 void UCEClonerComponent::OnLayoutNameChanged()
 {
 	if (!bEnabled)
@@ -2116,8 +2152,8 @@ void UCEClonerComponent::UpdateClonerMeshes()
 			{
 				if (const FCEClonerAttachmentItem* Item = ClonerTree.ItemAttachmentMap.Find(ClonerTree.RootActors[Idx]))
 				{
-					MeshProperties.Rotation = Item->ActorTransform.Rotator();
-					MeshProperties.Scale = Item->ActorTransform.GetScale3D();
+					MeshProperties.Rotation = GlobalRotation + Item->ActorTransform.Rotator();
+					MeshProperties.Scale = GlobalScale * Item->ActorTransform.GetScale3D();
 				}
 			}
 		}
diff --git a/Engine/Plugins/Experimental/ClonerEffector/Source/ClonerEffector/Private/Cloner/Extensions/CEClonerDisplacementExtension.cpp b/Engine/Plugins/Experimental/ClonerEffector/Source/ClonerEffector/Private/Cloner/Extensions/CEClonerDisplacementExtension.cpp
new file mode 100644
--- /dev/null
+++ b/Engine/Plugins/Experimental/ClonerEffector/Source/ClonerEffector/Private/Cloner/Extensions/CEClonerDisplacementExtension.cpp
@@ -0,0 +1,206 @@
+// Copyright Epic Games, Inc. All Rights Reserved.
+
+#include "Cloner/Extensions/CEClonerDisplacementExtension.h"
+
+#include "Cloner/CEClonerComponent.h"
+#include "Cloner/Layouts/CEClonerGridLayout.h"
+#include "NiagaraDataInterfaceTexture.h"
+
+void UCEClonerDisplacementExtension::SetDisplacementEnabled(bool bInEnabled)
+{
+	if (bDisplacementEnabled == bInEnabled)
+	{
+		return;
+	}
+
+	bDisplacementEnabled = bInEnabled;
+	MarkExtensionDirty();
+}
+
+void UCEClonerDisplacementExtension::SetDisplacementInvert(bool bInInvert)
+{
+	if (bDisplacementInvert == bInInvert)
+	{
+		return;
+	}
+
+	bDisplacementInvert = bInInvert;
+	MarkExtensionDirty();
+}
+
+void UCEClonerDisplacementExtension::SetDisplacementOffsetMax(const FVector& InMax)
+{
+	if (DisplacementOffsetMax.Equals(InMax))
+	{
+		return;
+	}
+
+	DisplacementOffsetMax = InMax;
+	MarkExtensionDirty();
+}
+
+void UCEClonerDisplacementExtension::SetDisplacementRotationMax(const FRotator& InMax)
+{
+	if (DisplacementRotationMax.Equals(InMax))
+	{
+		return;
+	}
+
+	DisplacementRotationMax = InMax;
+	MarkExtensionDirty();
+}
+
+void UCEClonerDisplacementExtension::SetDisplacementScaleMax(const FVector& InMax)
+{
+	if (DisplacementScaleMax.Equals(InMax))
+	{
+		return;
+	}
+
+	DisplacementScaleMax = InMax.ComponentMax(FVector(UE_KINDA_SMALL_NUMBER));
+	MarkExtensionDirty();
+}
+
+void UCEClonerDisplacementExtension::SetDisplacementTextureAsset(UTexture* InTexture)
+{
+	if (DisplacementTextureAsset == InTexture)
+	{
+		return;
+	}
+
+	DisplacementTextureAsset = InTexture;
+	MarkExtensionDirty();
+}
+
+void UCEClonerDisplacementExtension::SetDisplacementTextureSampleMode(ECEClonerTextureSampleChannel InSampleMode)
+{
+	if (DisplacementTextureSampleMode == InSampleMode)
+	{
+		return;
+	}
+
+	DisplacementTextureSampleMode = InSampleMode;
+	MarkExtensionDirty();
+}
+
+void UCEClonerDisplacementExtension::SetDisplacementTexturePlane(ECEClonerPlane InPlane)
+{
+	if (DisplacementTexturePlane == InPlane)
+	{
+		return;
+	}
+
+	DisplacementTexturePlane = InPlane;
+	MarkExtensionDirty();
+}
+
+void UCEClonerDisplacementExtension::SetDisplacementTextureOffset(const FVector2D& InOffset)
+{
+	if (DisplacementTextureOffset.Equals(InOffset))
+	{
+		return;
+	}
+
+	DisplacementTextureOffset = InOffset;
+	MarkExtensionDirty();
+}
+
+void UCEClonerDisplacementExtension::SetDisplacementTextureRotation(float InRotation)
+{
+	if (FMath::IsNearlyEqual(DisplacementTextureRotation, InRotation))
+	{
+		return;
+	}
+
+	DisplacementTextureRotation = InRotation;
+	MarkExtensionDirty();
+}
+
+void UCEClonerDisplacementExtension::SetDisplacementTextureScale(const FVector2D& InScale)
+{
+	if (DisplacementTextureScale.Equals(InScale))
+	{
+		return;
+	}
+
+	DisplacementTextureScale = InScale;
+	MarkExtensionDirty();
+}
+
+void UCEClonerDisplacementExtension::SetDisplacementTextureClamp(bool bInClamp)
+{
+	if (bDisplacementTextureClamp == bInClamp)
+	{
+		return;
+	}
+
+	bDisplacementTextureClamp = bInClamp;
+	MarkExtensionDirty();
+}
+
+#if WITH_EDITOR
+const TCEPropertyChangeDispatcher<UCEClonerDisplacementExtension> UCEClonerDisplacementExtension::PropertyChangeDispatcher =
+{
+	{ GET_MEMBER_NAME_CHECKED(UCEClonerDisplacementExtension, bDisplacementEnabled), &UCEClonerDisplacementExtension::OnExtensionPropertyChanged },
+	{ GET_MEMBER_NAME_CHECKED(UCEClonerDisplacementExtension, bDisplacementInvert), &UCEClonerDisplacementExtension::OnExtensionPropertyChanged },
+	{ GET_MEMBER_NAME_CHECKED(UCEClonerDisplacementExtension, DisplacementOffsetMax), &UCEClonerDisplacementExtension::OnExtensionPropertyChanged },
+	{ GET_MEMBER_NAME_CHECKED(UCEClonerDisplacementExtension, DisplacementRotationMax), &UCEClonerDisplacementExtension::OnExtensionPropertyChanged },
+	{ GET_MEMBER_NAME_CHECKED(UCEClonerDisplacementExtension, DisplacementScaleMax), &UCEClonerDisplacementExtension::OnExtensionPropertyChanged },
+	{ GET_MEMBER_NAME_CHECKED(UCEClonerDisplacementExtension, DisplacementTextureAsset), &UCEClonerDisplacementExtension::OnExtensionPropertyChanged },
+	{ GET_MEMBER_NAME_CHECKED(UCEClonerDisplacementExtension, DisplacementTextureSampleMode), &UCEClonerDisplacementExtension::OnExtensionPropertyChanged },
+	{ GET_MEMBER_NAME_CHECKED(UCEClonerDisplacementExtension, DisplacementTexturePlane), &UCEClonerDisplacementExtension::OnExtensionPropertyChanged },
+	{ GET_MEMBER_NAME_CHECKED(UCEClonerDisplacementExtension, DisplacementTextureOffset), &UCEClonerDisplacementExtension::OnExtensionPropertyChanged },
+	{ GET_MEMBER_NAME_CHECKED(UCEClonerDisplacementExtension, DisplacementTextureRotation), &UCEClonerDisplacementExtension::OnExtensionPropertyChanged },
+	{ GET_MEMBER_NAME_CHECKED(UCEClonerDisplacementExtension, DisplacementTextureScale), &UCEClonerDisplacementExtension::OnExtensionPropertyChanged },
+	{ GET_MEMBER_NAME_CHECKED(UCEClonerDisplacementExtension, bDisplacementTextureClamp), &UCEClonerDisplacementExtension::OnExtensionPropertyChanged },
+};
+
+void UCEClonerDisplacementExtension::PostEditChangeProperty(FPropertyChangedEvent& InPropertyChangedEvent)
+{
+	Super::PostEditChangeProperty(InPropertyChangedEvent);
+
+	PropertyChangeDispatcher.OnPropertyChanged(this, InPropertyChangedEvent);
+}
+#endif
+
+void UCEClonerDisplacementExtension::OnExtensionParametersChanged(UCEClonerComponent* InComponent)
+{
+	Super::OnExtensionParametersChanged(InComponent);
+
+	FNiagaraUserRedirectionParameterStore& ExposedParameters = InComponent->GetOverrideParameters();
+
+	InComponent->SetBoolParameter(TEXT("DisplacementEnabled"), bDisplacementEnabled);
+
+	InComponent->SetBoolParameter(TEXT("DisplacementInvert"), bDisplacementInvert);
+
+	InComponent->SetVectorParameter(TEXT("DisplacementOffsetMax"), DisplacementOffsetMax);
+	
+	InComponent->SetVectorParameter(TEXT("DisplacementRotationMax"), FVector(DisplacementRotationMax.Yaw, DisplacementRotationMax.Pitch, DisplacementRotationMax.Roll));
+	
+	InComponent->SetVectorParameter(TEXT("DisplacementScaleMax"), DisplacementScaleMax.ComponentMax(FVector(UE_KINDA_SMALL_NUMBER)));
+
+	static const FNiagaraVariable TextureVar(FNiagaraTypeDefinition(UNiagaraDataInterfaceTexture::StaticClass()), TEXT("DisplacementTexture"));
+	if (UNiagaraDataInterfaceTexture* TextureSamplerDI = Cast<UNiagaraDataInterfaceTexture>(ExposedParameters.GetDataInterface(TextureVar)))
+	{
+		TextureSamplerDI->SetTexture(DisplacementTextureAsset.Get());
+	}
+
+	static const FNiagaraVariable TexturePlaneVar(FNiagaraTypeDefinition(StaticEnum<ECEClonerPlane>()), TEXT("DisplacementTexturePlane"));
+	ExposedParameters.SetParameterValue<int32>(static_cast<int32>(DisplacementTexturePlane), TexturePlaneVar);
+
+	static const FNiagaraVariable TextureChannelVar(FNiagaraTypeDefinition(StaticEnum<ECEClonerTextureSampleChannel>()), TEXT("DisplacementTextureChannel"));
+	ExposedParameters.SetParameterValue<int32>(static_cast<int32>(DisplacementTextureSampleMode), TextureChannelVar);
+
+	InComponent->SetVariableVec2(TEXT("DisplacementTextureOffset"), DisplacementTextureOffset);
+
+	InComponent->SetFloatParameter(TEXT("DisplacementTextureRotation"), DisplacementTextureRotation);
+
+	InComponent->SetVariableVec2(TEXT("DisplacementTextureScale"), DisplacementTextureScale);
+
+	InComponent->SetVariableBool(TEXT("DisplacementTextureClamp"), bDisplacementTextureClamp);
+}
+
+bool UCEClonerDisplacementExtension::IsLayoutSupported(const UCEClonerLayoutBase* InLayout) const
+{
+	return InLayout->IsA<UCEClonerGridLayout>();
+}
\ No newline at end of file
diff --git a/Engine/Plugins/Experimental/ClonerEffector/Source/ClonerEffector/Private/Cloner/Extensions/CEClonerMeshRendererExtension.cpp b/Engine/Plugins/Experimental/ClonerEffector/Source/ClonerEffector/Private/Cloner/Extensions/CEClonerMeshRendererExtension.cpp
--- a/Engine/Plugins/Experimental/ClonerEffector/Source/ClonerEffector/Private/Cloner/Extensions/CEClonerMeshRendererExtension.cpp
+++ b/Engine/Plugins/Experimental/ClonerEffector/Source/ClonerEffector/Private/Cloner/Extensions/CEClonerMeshRendererExtension.cpp
@@ -251,6 +251,8 @@ void UCEClonerMeshRendererExtension::OnClonerMeshesUpdated()
 			UStaticMesh* DefaultMesh = NewDefaultMeshes[MeshIndex].Get();
 			FNiagaraMeshRendererMeshProperties& MeshProperties = MeshRenderer->Meshes[MeshIndex];
 			MeshProperties.Mesh = DefaultMesh && DefaultMesh->GetNumTriangles(0) > 0 ? DefaultMesh : nullptr;
+			MeshProperties.Scale = ClonerComponent->GetGlobalScale();
+			MeshProperties.Rotation = ClonerComponent->GetGlobalRotation();
 		}
 	}
 
diff --git a/Engine/Plugins/Experimental/ClonerEffector/Source/ClonerEffector/Public/Cloner/CEClonerComponent.h b/Engine/Plugins/Experimental/ClonerEffector/Source/ClonerEffector/Public/Cloner/CEClonerComponent.h
--- a/Engine/Plugins/Experimental/ClonerEffector/Source/ClonerEffector/Public/Cloner/CEClonerComponent.h
+++ b/Engine/Plugins/Experimental/ClonerEffector/Source/ClonerEffector/Public/Cloner/CEClonerComponent.h
@@ -115,6 +115,24 @@ public:
 		return Color;
 	}
 
+	UFUNCTION(BlueprintCallable, Category="Cloner")
+	CLONEREFFECTOR_API void SetGlobalScale(const FVector& InScale);
+
+	UFUNCTION(BlueprintPure, Category="Cloner")
+	const FVector& GetGlobalScale() const
+	{
+		return GlobalScale;
+	}
+
+	UFUNCTION(BlueprintCallable, Category="Cloner")
+	CLONEREFFECTOR_API void SetGlobalRotation(const FRotator& InRotation);
+
+	UFUNCTION(BlueprintPure, Category="Cloner")
+	const FRotator& GetGlobalRotation() const
+	{
+		return GlobalRotation;
+	}
+
 	UFUNCTION(BlueprintCallable, Category="Cloner")
 	CLONEREFFECTOR_API void SetLayoutName(FName InLayoutName);
 
@@ -178,7 +196,7 @@ public:
 
 #if WITH_EDITOR
 	/** This will force an update of the cloner attachment tree */
-	UFUNCTION(CallInEditor, Category="Cloner")
+	UFUNCTION(CallInEditor, Category="Utilities")
 	void ForceUpdateCloner();
 
 	/** Open project settings for cloner */
@@ -328,21 +346,35 @@ protected:
 	void OnActiveLayoutChanged();
 
 	/** Is this cloner enabled/disabled */
-	UPROPERTY(EditInstanceOnly, Setter="SetEnabled", Getter="GetEnabled", Category="Cloner")
+	UPROPERTY(EditInstanceOnly, Setter="SetEnabled", Getter="GetEnabled", Category=General)
 	bool bEnabled = true;
 
-	/** Interval to update the attachment tree and update the cloner meshes, 0 means each tick */
-	UPROPERTY(EditInstanceOnly, Setter, Getter, Category="Cloner", AdvancedDisplay, meta=(ClampMin="0"))
-	float TreeUpdateInterval = 0.f;
-
 	/** Cloner instance seed for random deterministic patterns */
-	UPROPERTY(EditInstanceOnly, Setter, Getter, Category="Cloner")
+	UPROPERTY(EditInstanceOnly, Setter, Getter, Category=General)
 	int32 Seed = 0;
 
 	/** Cloner color when unaffected by effectors, color will be passed down to the material (ParticleColor) */
-	UPROPERTY(EditInstanceOnly, Setter, Getter, Category="Cloner")
+	UPROPERTY(EditInstanceOnly, Setter, Getter, Category=General)
 	FLinearColor Color = FLinearColor::White;
 
+	/** Global scale applied on all cloned items */
+	UPROPERTY(EditInstanceOnly, Setter, Getter, Category=General, meta=(ClampMin="0", Delta="0.01", MotionDesignVectorWidget, AllowPreserveRatio="XYZ"))
+	FVector GlobalScale = FVector::OneVector;
+
+	/** Global rotation applied on all cloned items */
+	UPROPERTY(EditInstanceOnly, Setter, Getter, Category=General)
+	FRotator GlobalRotation = FRotator::ZeroRotator;
+
+	/** Interval to update the attachment tree and update the cloner meshes, 0 means each tick */
+	UPROPERTY(EditInstanceOnly, Setter, Getter, Category=General, AdvancedDisplay, meta=(ClampMin="0"))
+	float TreeUpdateInterval = 0.f;
+
+#if WITH_EDITORONLY_DATA
+	/** Toggle the sprite to visualize and click on this cloner */
+	UPROPERTY(EditInstanceOnly, AdvancedDisplay, Category=General, meta=(AllowPrivateAccess="true"))
+	bool bVisualizerSpriteVisible = true;
+#endif
+
 	/** Name of the layout to use */
 	UPROPERTY(EditInstanceOnly, Setter, Getter, DisplayName="Layout", Category="Layout", meta=(GetOptions="GetClonerLayoutNames"))
 	FName LayoutName = NAME_None;
@@ -363,12 +395,6 @@ protected:
 	UPROPERTY()
 	TArray<TObjectPtr<UCEClonerExtensionBase>> ExtensionInstances;
 
-#if WITH_EDITORONLY_DATA
-	/** Toggle the sprite to visualize and click on this cloner */
-	UPROPERTY(EditInstanceOnly, AdvancedDisplay, Category="Cloner", meta=(AllowPrivateAccess="true"))
-	bool bVisualizerSpriteVisible = true;
-#endif
-
 private:
 	static constexpr TCHAR SpriteTexturePath[] = TEXT("/Script/Engine.Texture2D'/ClonerEffector/Textures/T_ClonerIcon.T_ClonerIcon'");
 
@@ -395,6 +421,8 @@ private:
 
 	void OnSeedChanged();
 	void OnColorChanged();
+	void OnGlobalScaleChanged();
+	void OnGlobalRotationChanged();
 	void OnLayoutNameChanged();
 
 #if WITH_EDITOR
diff --git a/Engine/Plugins/Experimental/ClonerEffector/Source/ClonerEffector/Public/Cloner/Extensions/CEClonerConstraintExtension.h b/Engine/Plugins/Experimental/ClonerEffector/Source/ClonerEffector/Public/Cloner/Extensions/CEClonerConstraintExtension.h
--- a/Engine/Plugins/Experimental/ClonerEffector/Source/ClonerEffector/Public/Cloner/Extensions/CEClonerConstraintExtension.h
+++ b/Engine/Plugins/Experimental/ClonerEffector/Source/ClonerEffector/Public/Cloner/Extensions/CEClonerConstraintExtension.h
@@ -216,7 +216,7 @@ protected:
 	ECEClonerCompareMode TextureCompareMode = ECEClonerCompareMode::Greater;
 
 	UPROPERTY(EditInstanceOnly, Setter, Getter, Category="Constraint", meta=(ClampMin="0", EditCondition="Constraint == ECEClonerGridConstraint::Texture", EditConditionHides))
-	float TextureThreshold = 0.f;
+	float TextureThreshold = 0.1f;
 
 	UPROPERTY(EditInstanceOnly, Setter, Getter, Category="Constraint", meta=(EditCondition="Constraint == ECEClonerGridConstraint::Texture", EditConditionHides))
 	FVector2D TextureOffset = FVector2D::ZeroVector;
diff --git a/Engine/Plugins/Experimental/ClonerEffector/Source/ClonerEffector/Public/Cloner/Extensions/CEClonerDisplacementExtension.h b/Engine/Plugins/Experimental/ClonerEffector/Source/ClonerEffector/Public/Cloner/Extensions/CEClonerDisplacementExtension.h
new file mode 100644
--- /dev/null
+++ b/Engine/Plugins/Experimental/ClonerEffector/Source/ClonerEffector/Public/Cloner/Extensions/CEClonerDisplacementExtension.h
@@ -0,0 +1,183 @@
+// Copyright Epic Games, Inc. All Rights Reserved.
+
+#pragma once
+
+#include "CEClonerExtensionBase.h"
+#include "CEPropertyChangeDispatcher.h"
+#include "CEClonerDisplacementExtension.generated.h"
+
+/** Extension dealing with clone grid displacements */
+UCLASS(MinimalAPI, BlueprintType, Within=CEClonerComponent, meta=(Section="Cloner", Priority=0))
+class UCEClonerDisplacementExtension : public UCEClonerExtensionBase
+{
+	GENERATED_BODY()
+
+public:
+	UCEClonerDisplacementExtension()
+		: UCEClonerExtensionBase(
+			TEXT("Displacement")
+			, 0
+		)
+	{}
+
+	UFUNCTION(BlueprintCallable, Category="Cloner")
+	CLONEREFFECTOR_API void SetDisplacementEnabled(bool bInEnabled);
+
+	UFUNCTION(BlueprintPure, Category="Cloner")
+	bool GetDisplacementEnabled() const
+	{
+		return bDisplacementEnabled;
+	}
+
+	UFUNCTION(BlueprintCallable, Category="Cloner")
+	CLONEREFFECTOR_API void SetDisplacementInvert(bool bInInvert);
+
+	UFUNCTION(BlueprintPure, Category="Cloner")
+	bool GetDisplacementInvert() const
+	{
+		return bDisplacementInvert;
+	}
+	
+	UFUNCTION(BlueprintCallable, Category="Cloner")
+	CLONEREFFECTOR_API void SetDisplacementOffsetMax(const FVector& InMax);
+
+	UFUNCTION(BlueprintPure, Category="Cloner")
+	const FVector& GetDisplacementOffsetMax() const
+	{
+		return DisplacementOffsetMax;
+	}
+	
+	UFUNCTION(BlueprintCallable, Category="Cloner")
+	CLONEREFFECTOR_API void SetDisplacementRotationMax(const FRotator& InMax);
+
+	UFUNCTION(BlueprintPure, Category="Cloner")
+	const FRotator& GetDisplacementRotationMax() const
+	{
+		return DisplacementRotationMax;
+	}
+	
+	UFUNCTION(BlueprintCallable, Category="Cloner")
+	CLONEREFFECTOR_API void SetDisplacementScaleMax(const FVector& InMax);
+
+	UFUNCTION(BlueprintPure, Category="Cloner")
+	const FVector& GetDisplacementScaleMax() const
+	{
+		return DisplacementScaleMax;
+	}
+
+	UFUNCTION(BlueprintCallable, Category="Cloner")
+	CLONEREFFECTOR_API void SetDisplacementTextureAsset(UTexture* InTexture);
+
+	UFUNCTION(BlueprintPure, Category="Cloner")
+	UTexture* GetDisplacementTextureAsset() const
+	{
+		return DisplacementTextureAsset;
+	}
+
+	UFUNCTION(BlueprintCallable, Category="Cloner")
+	CLONEREFFECTOR_API void SetDisplacementTextureSampleMode(ECEClonerTextureSampleChannel InSampleMode);
+
+	UFUNCTION(BlueprintPure, Category="Cloner")
+	ECEClonerTextureSampleChannel GetDisplacementTextureSampleMode() const
+	{
+		return DisplacementTextureSampleMode;
+	}
+
+	UFUNCTION(BlueprintCallable, Category="Cloner")
+	CLONEREFFECTOR_API void SetDisplacementTexturePlane(ECEClonerPlane InPlane);
+
+	UFUNCTION(BlueprintPure, Category="Cloner")
+	ECEClonerPlane GetDisplacementTexturePlane() const
+	{
+		return DisplacementTexturePlane;
+	}
+
+	UFUNCTION(BlueprintCallable, Category="Cloner")
+	CLONEREFFECTOR_API void SetDisplacementTextureOffset(const FVector2D& InOffset);
+
+	UFUNCTION(BlueprintPure, Category="Cloner")
+	FVector2D GetDisplacementTextureOffset() const
+	{
+		return DisplacementTextureOffset;
+	}
+
+	UFUNCTION(BlueprintCallable, Category="Cloner")
+	CLONEREFFECTOR_API void SetDisplacementTextureRotation(float InRotation);
+
+	UFUNCTION(BlueprintPure, Category="Cloner")
+	float GetDisplacementTextureRotation() const
+	{
+		return DisplacementTextureRotation;
+	}
+	
+	UFUNCTION(BlueprintCallable, Category="Cloner")
+	CLONEREFFECTOR_API void SetDisplacementTextureScale(const FVector2D& InScale);
+
+	UFUNCTION(BlueprintPure, Category="Cloner")
+	FVector2D GetDisplacementTextureScale() const
+	{
+		return DisplacementTextureScale;
+	}
+	
+	UFUNCTION(BlueprintCallable, Category="Cloner")
+	CLONEREFFECTOR_API void SetDisplacementTextureClamp(bool bInClamp);
+
+	UFUNCTION(BlueprintPure, Category="Cloner")
+	bool GetDisplacementTextureClamp() const
+	{
+		return bDisplacementTextureClamp;
+	}
+
+protected:
+	//~ Begin UObject
+#if WITH_EDITOR
+	virtual void PostEditChangeProperty(FPropertyChangedEvent& InPropertyChangedEvent) override;
+#endif
+	//~ End UObject
+
+	//~ Begin UCEClonerExtensionBase
+	virtual void OnExtensionParametersChanged(UCEClonerComponent* InComponent) override;
+	virtual bool IsLayoutSupported(const UCEClonerLayoutBase* InLayout) const override;
+	//~ End UCEClonerExtensionBase
+
+	UPROPERTY(EditInstanceOnly, Setter="SetDisplacementEnabled", Getter="GetDisplacementEnabled", DisplayName="Enabled", Category="Displacement")
+	bool bDisplacementEnabled = false;
+
+	UPROPERTY(EditInstanceOnly, Setter="SetDisplacementInvert", Getter="GetDisplacementInvert", DisplayName="Invert", Category="Displacement", meta=(EditCondition="bDisplacementEnabled", EditConditionHides))
+	bool bDisplacementInvert = false;
+
+	UPROPERTY(EditInstanceOnly, Setter, Getter, DisplayName="Offset", Category="Displacement", meta=(EditCondition="bDisplacementEnabled", EditConditionHides))
+	FVector DisplacementOffsetMax = FVector(0, 0, 100.f);
+
+	UPROPERTY(EditInstanceOnly, Setter, Getter, DisplayName="Rotation", Category="Displacement", meta=(EditCondition="bDisplacementEnabled", EditConditionHides))
+	FRotator DisplacementRotationMax = FRotator::ZeroRotator;
+
+	UPROPERTY(EditInstanceOnly, Setter, Getter, DisplayName="Scale", Category="Displacement", meta=(ClampMin="0", Delta="0.01", MotionDesignVectorWidget, AllowPreserveRatio="XYZ", EditCondition="bDisplacementEnabled", EditConditionHides))
+	FVector DisplacementScaleMax = FVector::OneVector;
+
+	UPROPERTY(EditInstanceOnly, Setter, Getter, DisplayName="TextureAsset", Category="Displacement", meta=(EditCondition="bDisplacementEnabled", EditConditionHides))
+	TObjectPtr<UTexture> DisplacementTextureAsset;
+
+	UPROPERTY(EditInstanceOnly, Setter, Getter, DisplayName="TexturePlane", Category="Displacement", meta=(InvalidEnumValues="Custom", EditCondition="bDisplacementEnabled", EditConditionHides))
+	ECEClonerPlane DisplacementTexturePlane = ECEClonerPlane::XY;
+
+	UPROPERTY(EditInstanceOnly, Setter, Getter, DisplayName="TextureSampleMode", Category="Displacement", meta=(EditCondition="bDisplacementEnabled", EditConditionHides))
+	ECEClonerTextureSampleChannel DisplacementTextureSampleMode = ECEClonerTextureSampleChannel::RGBLuminance;
+
+	UPROPERTY(EditInstanceOnly, Setter, Getter, DisplayName="TextureOffset", Category="Displacement", meta=(EditCondition="bDisplacementEnabled", EditConditionHides))
+	FVector2D DisplacementTextureOffset = FVector2D::ZeroVector;
+	
+	UPROPERTY(EditInstanceOnly, Setter, Getter, DisplayName="TextureRotation", Category="Displacement", meta=(EditCondition="bDisplacementEnabled", EditConditionHides))
+	float DisplacementTextureRotation = 0.f;
+	
+	UPROPERTY(EditInstanceOnly, Setter, Getter, DisplayName="TextureScale", Category="Displacement", meta=(Delta="0.01", MotionDesignVectorWidget, AllowPreserveRatio="XY", EditCondition="bDisplacementEnabled", EditConditionHides))
+	FVector2D DisplacementTextureScale = FVector2D(1.0);
+
+	UPROPERTY(EditInstanceOnly, Setter="SetDisplacementTextureClamp", Getter="GetDisplacementTextureClamp", DisplayName="TextureClamp", Category="Displacement", meta=(EditCondition="bDisplacementEnabled", EditConditionHides))
+	bool bDisplacementTextureClamp = false;
+
+private:
+#if WITH_EDITOR
+	static const TCEPropertyChangeDispatcher<UCEClonerDisplacementExtension> PropertyChangeDispatcher;
+#endif
+};
\ No newline at end of file
diff --git a/Engine/Plugins/Experimental/ClonerEffector/Source/ClonerEffectorEditor/Private/Cloner/Customizations/CEEditorClonerComponentDetailCustomization.cpp b/Engine/Plugins/Experimental/ClonerEffector/Source/ClonerEffectorEditor/Private/Cloner/Customizations/CEEditorClonerComponentDetailCustomization.cpp
--- a/Engine/Plugins/Experimental/ClonerEffector/Source/ClonerEffectorEditor/Private/Cloner/Customizations/CEEditorClonerComponentDetailCustomization.cpp
+++ b/Engine/Plugins/Experimental/ClonerEffector/Source/ClonerEffectorEditor/Private/Cloner/Customizations/CEEditorClonerComponentDetailCustomization.cpp
@@ -43,6 +43,15 @@ void FCEEditorClonerComponentDetailCustomization::CustomizeDetails(IDetailLayout
 		LayoutCategoryBuilder.AddProperty(LayoutHandle);
 	}
 
+	{
+		IDetailCategoryBuilder& GeneralCategory = InDetailBuilder.EditCategory(TEXT("General"), LOCTEXT("GeneralCategory", "General"), ECategoryPriority::Important);
+		IDetailCategoryBuilder& TransformCategory = InDetailBuilder.EditCategory(TEXT("Transform"), LOCTEXT("TransformCategory", "Transform"), ECategoryPriority::Important);
+		TransformCategory.SetSortOrder(GeneralCategory.GetSortOrder() + 1);
+
+		const TSharedRef<FPropertySection> GeneralSection = PropertyModule.FindOrCreateSection(ComponentClassName, TEXT("General"), FText::FromName(TEXT("General")));
+		GeneralSection->AddCategory(TEXT("General"));
+	}
+
 	// Everything needs to be below Cloner category
 	const IDetailCategoryBuilder& ClonerCategoryBuilder = InDetailBuilder.EditCategory(TEXT("Cloner"), FText::FromName(TEXT("Cloner")));
 	const int32 ExtensionSortOrderOffset = 1;
