MotionDesign : ClonerEffector
- Prevent slate from throttling when interactively changing a property value to allow live preview for cloner and effector

diff --git a/Engine/Plugins/Experimental/ClonerEffector/Source/ClonerEffectorEditor/Private/CEEditorModule.cpp b/Engine/Plugins/Experimental/ClonerEffector/Source/ClonerEffectorEditor/Private/CEEditorModule.cpp
--- a/Engine/Plugins/Experimental/ClonerEffector/Source/ClonerEffectorEditor/Private/CEEditorModule.cpp
+++ b/Engine/Plugins/Experimental/ClonerEffector/Source/ClonerEffectorEditor/Private/CEEditorModule.cpp
@@ -2,6 +2,7 @@
 
 #include "CEEditorModule.h"
 
+#include "CEEditorThrottleManager.h"
 #include "Cloner/CEClonerComponent.h"
 #include "Cloner/CEClonerActor.h"
 #include "Cloner/Customizations/CEEditorClonerActorDetailCustomization.h"
@@ -49,35 +50,42 @@ void FCEEditorModule::StartupModule()
 	// Custom cloner track
 	ISequencerModule& SequencerModule = FModuleManager::LoadModuleChecked<ISequencerModule>(TEXT("Sequencer"));
 	ClonerTrackCreateEditorHandle = SequencerModule.RegisterTrackEditor(FOnCreateTrackEditor::CreateStatic(&FMovieSceneClonerTrackEditor::CreateTrackEditor));
+
+	// Disable slate throttling for interactive changes
+	ThrottleManager = MakeShared<FCEEditorThrottleManager>();
+	ThrottleManager->Init();
 }
 
 void FCEEditorModule::ShutdownModule()
 {
-	if (FModuleManager::Get().IsModuleLoaded("PropertyEditor") && UObjectInitialized())
+	if (UObjectInitialized())
 	{
-		FPropertyEditorModule& PropertyModule = FModuleManager::GetModuleChecked<FPropertyEditorModule>("PropertyEditor");
-
-		// Cloner customization
-		PropertyModule.UnregisterCustomClassLayout(ACEClonerActor::StaticClass()->GetFName());
-		PropertyModule.UnregisterCustomClassLayout(UCEClonerComponent::StaticClass()->GetFName());
-		PropertyModule.UnregisterCustomClassLayout(UCEClonerEffectorExtension::StaticClass()->GetFName());
-		PropertyModule.UnregisterCustomClassLayout(UCEClonerSplineLayout::StaticClass()->GetFName());
-		PropertyModule.UnregisterCustomClassLayout(UCEClonerMeshLayout::StaticClass()->GetFName());
-		PropertyModule.UnregisterCustomClassLayout(UCEClonerLifetimeExtension::StaticClass()->GetFName());
+		if (FPropertyEditorModule* PropertyModule = FModuleManager::GetModulePtr<FPropertyEditorModule>("PropertyEditor"))
+		{
+			// Cloner customization
+			PropertyModule->UnregisterCustomClassLayout(ACEClonerActor::StaticClass()->GetFName());
+			PropertyModule->UnregisterCustomClassLayout(UCEClonerComponent::StaticClass()->GetFName());
+			PropertyModule->UnregisterCustomClassLayout(UCEClonerEffectorExtension::StaticClass()->GetFName());
+			PropertyModule->UnregisterCustomClassLayout(UCEClonerSplineLayout::StaticClass()->GetFName());
+			PropertyModule->UnregisterCustomClassLayout(UCEClonerMeshLayout::StaticClass()->GetFName());
+			PropertyModule->UnregisterCustomClassLayout(UCEClonerLifetimeExtension::StaticClass()->GetFName());
 
-		// Effector customization
-		PropertyModule.UnregisterCustomClassLayout(ACEEffectorActor::StaticClass()->GetFName());
-		PropertyModule.UnregisterCustomClassLayout(UCEEffectorComponent::StaticClass()->GetFName());
-		PropertyModule.UnregisterCustomClassLayout(UCEEffectorBoundType::StaticClass()->GetFName());
+			// Effector customization
+			PropertyModule->UnregisterCustomClassLayout(ACEEffectorActor::StaticClass()->GetFName());
+			PropertyModule->UnregisterCustomClassLayout(UCEEffectorComponent::StaticClass()->GetFName());
+			PropertyModule->UnregisterCustomClassLayout(UCEEffectorBoundType::StaticClass()->GetFName());
+		}
 	}
 
 	// Custom cloner track
-	if (FModuleManager::Get().IsModuleLoaded("Sequencer"))
+	if (ISequencerModule* SequencerModule = FModuleManager::GetModulePtr<ISequencerModule>(TEXT("Sequencer")))
 	{
-		ISequencerModule& SequencerModule = FModuleManager::GetModuleChecked<ISequencerModule>(TEXT("Sequencer"));
-		SequencerModule.UnRegisterTrackEditor(ClonerTrackCreateEditorHandle);
+		SequencerModule->UnRegisterTrackEditor(ClonerTrackCreateEditorHandle);
 		ClonerTrackCreateEditorHandle.Reset();
 	}
+
+	ThrottleManager.Reset();
+	ThrottleManager = nullptr;
 }
 
-IMPLEMENT_MODULE(FCEEditorModule, ClonerEffectorEditor)
\ No newline at end of file
+IMPLEMENT_MODULE(FCEEditorModule, ClonerEffectorEditor)
diff --git a/Engine/Plugins/Experimental/ClonerEffector/Source/ClonerEffectorEditor/Private/CEEditorModule.h b/Engine/Plugins/Experimental/ClonerEffector/Source/ClonerEffectorEditor/Private/CEEditorModule.h
--- a/Engine/Plugins/Experimental/ClonerEffector/Source/ClonerEffectorEditor/Private/CEEditorModule.h
+++ b/Engine/Plugins/Experimental/ClonerEffector/Source/ClonerEffectorEditor/Private/CEEditorModule.h
@@ -4,6 +4,8 @@
 
 #include "Modules/ModuleManager.h"
 
+class FCEEditorThrottleManager;
+
 class FCEEditorModule : public IModuleInterface
 {
 public:
@@ -14,4 +16,5 @@ public:
 
 protected:
 	FDelegateHandle ClonerTrackCreateEditorHandle;
+	TSharedPtr<FCEEditorThrottleManager> ThrottleManager;
 };
diff --git a/Engine/Plugins/Experimental/ClonerEffector/Source/ClonerEffectorEditor/Private/CEEditorThrottleManager.cpp b/Engine/Plugins/Experimental/ClonerEffector/Source/ClonerEffectorEditor/Private/CEEditorThrottleManager.cpp
new file mode 100644
--- /dev/null
+++ b/Engine/Plugins/Experimental/ClonerEffector/Source/ClonerEffectorEditor/Private/CEEditorThrottleManager.cpp
@@ -0,0 +1,78 @@
+// Copyright Epic Games, Inc. All Rights Reserved.
+
+#include "CEEditorThrottleManager.h"
+
+#include "Application/ThrottleManager.h"
+#include "Cloner/CEClonerActor.h"
+#include "Effector/CEEffectorActor.h"
+#include "Misc/Optional.h"
+#include "Templates/SharedPointer.h"
+#include "UObject/Object.h"
+#include "UObject/UObjectGlobals.h"
+
+FCEEditorThrottleScope::FCEEditorThrottleScope(FName InPropertyName)
+	: PropertyName(InPropertyName)
+{
+	FSlateThrottleManager::Get().DisableThrottle(true);
+}
+
+FCEEditorThrottleScope::~FCEEditorThrottleScope()
+{
+	FSlateThrottleManager::Get().DisableThrottle(false);
+}
+
+FCEEditorThrottleManager::~FCEEditorThrottleManager()
+{
+	Stop();
+}
+
+void FCEEditorThrottleManager::Init()
+{
+	FCoreUObjectDelegates::OnObjectPropertyChanged.AddSP(this, &FCEEditorThrottleManager::OnPostPropertyChanged);
+}
+
+void FCEEditorThrottleManager::Stop()
+{
+	FCoreUObjectDelegates::OnObjectPropertyChanged.RemoveAll(this);
+}
+
+void FCEEditorThrottleManager::OnPostPropertyChanged(UObject* InObject, FPropertyChangedEvent& InEvent)
+{
+	if (!IsValid(InObject))
+	{
+		return;
+	}
+
+	const AActor* Owner = InObject->GetTypedOuter<AActor>();
+	if (!Owner || (!Owner->IsA<ACEClonerActor>() && !Owner->IsA<ACEEffectorActor>()))
+	{
+		return;
+	}
+
+	// Due to multiple events chaining, use property name to filter events
+	const FName PropertyName = InEvent.GetMemberPropertyName();
+	if (PropertyName.IsNone())
+	{
+		return;
+	}
+
+	// Disable throttling after first interactive event and enable it back once value is set
+	if (InEvent.ChangeType == EPropertyChangeType::Interactive)
+	{
+		if (!ThrottleScope.IsSet())
+		{
+			// Disable throttling
+			ThrottleScope.Emplace(PropertyName);
+		}
+
+		return;
+	}
+
+	if (!ThrottleScope.IsSet() || !ThrottleScope->GetPropertyName().IsEqual(PropertyName))
+	{
+		return;
+	}
+
+	// Enable throttling
+	ThrottleScope.Reset();
+}
diff --git a/Engine/Plugins/Experimental/ClonerEffector/Source/ClonerEffectorEditor/Private/CEEditorThrottleManager.h b/Engine/Plugins/Experimental/ClonerEffector/Source/ClonerEffectorEditor/Private/CEEditorThrottleManager.h
new file mode 100644
--- /dev/null
+++ b/Engine/Plugins/Experimental/ClonerEffector/Source/ClonerEffectorEditor/Private/CEEditorThrottleManager.h
@@ -0,0 +1,46 @@
+// Copyright Epic Games, Inc. All Rights Reserved.
+
+#pragma once
+
+#include "Misc/OptionalFwd.h"
+#include "Templates/SharedPointer.h"
+#include "UObject/NameTypes.h"
+
+class UObject;
+struct FPropertyChangedEvent;
+
+struct FCEEditorThrottleScope
+{
+	explicit FCEEditorThrottleScope(FName InPropertyName);
+
+	virtual ~FCEEditorThrottleScope();
+
+	FName GetPropertyName() const
+	{
+		return PropertyName;
+	}
+
+private:
+	/** Property that is updating */
+	FName PropertyName;
+};
+
+/**
+ * Used to allow preview when interactively changing a property within cloner/effector
+ * Disable slate throttling to allow viewport updates
+ */
+class FCEEditorThrottleManager : public TSharedFromThis<FCEEditorThrottleManager>
+{
+public:
+	FCEEditorThrottleManager() = default;
+	virtual ~FCEEditorThrottleManager();
+
+	void Init();
+
+	void Stop();
+
+private:
+	void OnPostPropertyChanged(UObject* InObject, FPropertyChangedEvent& InEvent);
+
+	TOptional<FCEEditorThrottleScope> ThrottleScope;
+};
\ No newline at end of file
