MotionDesign : ClonerEffector
- Fix cloner not appearing in quick launch (runtime) package due to resources not ready

diff --git a/Engine/Plugins/Experimental/ClonerEffector/Source/ClonerEffector/Private/Cloner/CEClonerComponent.cpp b/Engine/Plugins/Experimental/ClonerEffector/Source/ClonerEffector/Private/Cloner/CEClonerComponent.cpp
--- a/Engine/Plugins/Experimental/ClonerEffector/Source/ClonerEffector/Private/Cloner/CEClonerComponent.cpp
+++ b/Engine/Plugins/Experimental/ClonerEffector/Source/ClonerEffector/Private/Cloner/CEClonerComponent.cpp
@@ -1178,11 +1178,8 @@ void UCEClonerComponent::InitializeCloner()
 	}
 #endif
 
-	const UWorld* World = GetWorld();
-	const ULevel* Level = GetComponentLevel();
-
 	// When level is streamed in, wait until actor hierarchy and resources are ready before initializing
-	if (IsValid(World) && IsValid(Level))
+	if (const ULevel* Level = GetComponentLevel())
 	{
 		if (!Level->IsPersistentLevel() && Level->HasAnyInternalFlags(EInternalObjectFlags_AsyncLoading))
 		{
@@ -1198,12 +1195,8 @@ void UCEClonerComponent::InitializeCloner()
 void UCEClonerComponent::OnLevelLoaded(ULevel* InLevel, UWorld* InWorld)
 {
 	const ULevel* Level = GetComponentLevel();
-	const UWorld* World = GetWorld();
 
-	if (!IsValid(Level)
-		|| !IsValid(World)
-		|| Level != InLevel
-		|| Level->HasAnyInternalFlags(EInternalObjectFlags_AsyncLoading))
+	if (!IsValid(Level) || Level != InLevel)
 	{
 		return;
 	}
@@ -1220,9 +1213,6 @@ void UCEClonerComponent::PostInitializeCloner()
 
 	RegisterTicker();
 
-	// Load layout after registering ticker to let attachment tree update, then layout rendering will occur with up to date attachment data
-	OnEnabledChanged();
-
 	OnClonerInitializedDelegate.Broadcast(this);
 }
 
@@ -1238,6 +1228,40 @@ void UCEClonerComponent::RegisterTicker()
 	ClonerTickerHandle = FTSTicker::GetCoreTicker().AddTicker(FTickerDelegate::CreateUObject(this, &UCEClonerComponent::TickCloner));
 }
 
+bool UCEClonerComponent::CheckResourcesReady()
+{
+	const ULevel* Level = GetComponentLevel();
+
+	// Check level is not async loading
+	if (!IsValid(Level) || Level->HasAnyInternalFlags(EInternalObjectFlags_AsyncLoading))
+	{
+		return false;
+	}
+
+	// Check cached meshes are not async loading
+	for (const TObjectPtr<UStaticMesh>& BakedMesh : ClonerTree.MergedBakedMeshes)
+	{
+		if (BakedMesh && BakedMesh->HasAnyInternalFlags(EInternalObjectFlags_AsyncLoading))
+		{
+			return false;
+		}
+	}
+
+	// Check cached active system is not async loading
+	for (const TObjectPtr<UCEClonerLayoutBase>& Layout : LayoutInstances)
+	{
+		if (Layout && Layout->GetLayoutName() == LayoutName)
+		{
+			if (Layout->GetSystem() && Layout->GetSystem()->HasAnyInternalFlags(EInternalObjectFlags_AsyncLoading))
+			{
+				return false;
+			}
+		}
+	}
+
+	return true;
+}
+
 bool UCEClonerComponent::TickCloner(float InDelta)
 {
 	if (!bClonerInitialized)
@@ -1245,6 +1269,17 @@ bool UCEClonerComponent::TickCloner(float InDelta)
 		return false;
 	}
 
+	if (!bClonerResourcesReady)
+	{
+		if (!CheckResourcesReady())
+		{
+			return true;
+		}
+
+		OnEnabledChanged();
+		bClonerResourcesReady = true;
+	}
+
 	if (bEnabled)
 	{
 		TreeUpdateDeltaTime += InDelta;
@@ -2138,12 +2173,6 @@ void UCEClonerComponent::UpdateClonerMeshes()
 		{
 			UStaticMesh* StaticMesh = ClonerTree.MergedBakedMeshes[Idx].Get();
 
-			// Make sure static mesh is loaded and available when streaming
-			if (StaticMesh && StaticMesh->HasAnyInternalFlags(EInternalObjectFlags_AsyncLoading))
-			{
-				return;
-			}
-
 			FNiagaraMeshRendererMeshProperties& MeshProperties = !MeshRenderer->Meshes.IsValidIndex(Idx) ? MeshRenderer->Meshes.AddDefaulted_GetRef() : MeshRenderer->Meshes[Idx];
 			bMeshChanged |= MeshProperties.Mesh != StaticMesh;
 			MeshProperties.Mesh = StaticMesh && StaticMesh->GetNumTriangles(0) > 0 ? StaticMesh : nullptr;
diff --git a/Engine/Plugins/Experimental/ClonerEffector/Source/ClonerEffector/Public/Cloner/CEClonerComponent.h b/Engine/Plugins/Experimental/ClonerEffector/Source/ClonerEffector/Public/Cloner/CEClonerComponent.h
--- a/Engine/Plugins/Experimental/ClonerEffector/Source/ClonerEffector/Public/Cloner/CEClonerComponent.h
+++ b/Engine/Plugins/Experimental/ClonerEffector/Source/ClonerEffector/Public/Cloner/CEClonerComponent.h
@@ -403,6 +403,7 @@ private:
 	void OnLevelLoaded(ULevel* InLevel, UWorld* InWorld);
 	void PostInitializeCloner();
 	void RegisterTicker();
+	bool CheckResourcesReady();
 	bool TickCloner(float InDelta);
 
 	/**
@@ -504,6 +505,8 @@ private:
 
 	bool bClonerInitialized = false;
 
+	bool bClonerResourcesReady = false;
+
 	FTSTicker::FDelegateHandle ClonerTickerHandle;
 
 #if WITH_EDITOR
