MotionDesign : ClonerEffector
- Add texture extension for grid layout thats sends texture, uv parameters to material, pick between constraint, displacement or custom provider
- Cloner material function and default material adapted to handle both texture and color

#virtualized

diff --git a/Engine/Plugins/Experimental/ClonerEffector/Source/ClonerEffector/Private/Cloner/Extensions/CEClonerCollisionExtension.cpp b/Engine/Plugins/Experimental/ClonerEffector/Source/ClonerEffector/Private/Cloner/Extensions/CEClonerCollisionExtension.cpp
--- a/Engine/Plugins/Experimental/ClonerEffector/Source/ClonerEffector/Private/Cloner/Extensions/CEClonerCollisionExtension.cpp
+++ b/Engine/Plugins/Experimental/ClonerEffector/Source/ClonerEffector/Private/Cloner/Extensions/CEClonerCollisionExtension.cpp
@@ -9,6 +9,13 @@
 #include "NiagaraMeshRendererProperties.h"
 #include "NiagaraSystem.h"
 
+UCEClonerCollisionExtension::UCEClonerCollisionExtension()
+	: UCEClonerExtensionBase(
+			TEXT("Collisions")
+			, 0
+		)
+{}
+
 void UCEClonerCollisionExtension::SetSurfaceCollisionEnabled(bool bInSurfaceCollisionEnabled)
 {
 	if (bSurfaceCollisionEnabled == bInSurfaceCollisionEnabled)
diff --git a/Engine/Plugins/Experimental/ClonerEffector/Source/ClonerEffector/Private/Cloner/Extensions/CEClonerConstraintExtension.cpp b/Engine/Plugins/Experimental/ClonerEffector/Source/ClonerEffector/Private/Cloner/Extensions/CEClonerConstraintExtension.cpp
--- a/Engine/Plugins/Experimental/ClonerEffector/Source/ClonerEffector/Private/Cloner/Extensions/CEClonerConstraintExtension.cpp
+++ b/Engine/Plugins/Experimental/ClonerEffector/Source/ClonerEffector/Private/Cloner/Extensions/CEClonerConstraintExtension.cpp
@@ -7,6 +7,13 @@
 #include "NiagaraDataInterfaceTexture.h"
 #include "NiagaraSystem.h"
 
+UCEClonerConstraintExtension::UCEClonerConstraintExtension()
+	: UCEClonerExtensionBase(
+		TEXT("Constraint")
+		, 0
+	)
+{}
+
 void UCEClonerConstraintExtension::SetConstraint(ECEClonerGridConstraint InConstraint)
 {
 	if (Constraint == InConstraint)
diff --git a/Engine/Plugins/Experimental/ClonerEffector/Source/ClonerEffector/Private/Cloner/Extensions/CEClonerDisplacementExtension.cpp b/Engine/Plugins/Experimental/ClonerEffector/Source/ClonerEffector/Private/Cloner/Extensions/CEClonerDisplacementExtension.cpp
--- a/Engine/Plugins/Experimental/ClonerEffector/Source/ClonerEffector/Private/Cloner/Extensions/CEClonerDisplacementExtension.cpp
+++ b/Engine/Plugins/Experimental/ClonerEffector/Source/ClonerEffector/Private/Cloner/Extensions/CEClonerDisplacementExtension.cpp
@@ -6,6 +6,13 @@
 #include "Cloner/Layouts/CEClonerGridLayout.h"
 #include "NiagaraDataInterfaceTexture.h"
 
+UCEClonerDisplacementExtension::UCEClonerDisplacementExtension()
+	: UCEClonerExtensionBase(
+		TEXT("Displacement")
+		, 0
+	)
+{}
+
 void UCEClonerDisplacementExtension::SetDisplacementEnabled(bool bInEnabled)
 {
 	if (bDisplacementEnabled == bInEnabled)
diff --git a/Engine/Plugins/Experimental/ClonerEffector/Source/ClonerEffector/Private/Cloner/Extensions/CEClonerEffectorExtension.cpp b/Engine/Plugins/Experimental/ClonerEffector/Source/ClonerEffector/Private/Cloner/Extensions/CEClonerEffectorExtension.cpp
--- a/Engine/Plugins/Experimental/ClonerEffector/Source/ClonerEffector/Private/Cloner/Extensions/CEClonerEffectorExtension.cpp
+++ b/Engine/Plugins/Experimental/ClonerEffector/Source/ClonerEffector/Private/Cloner/Extensions/CEClonerEffectorExtension.cpp
@@ -23,6 +23,13 @@ FName UCEClonerEffectorExtension::GetEffectorActorsWeakName()
 }
 #endif
 
+UCEClonerEffectorExtension::UCEClonerEffectorExtension()
+	: UCEClonerExtensionBase(
+		TEXT("Effector")
+		, 0
+	)
+{}
+
 bool UCEClonerEffectorExtension::LinkEffector(AActor* InEffectorActor)
 {
 	if (!IsValid(InEffectorActor)
diff --git a/Engine/Plugins/Experimental/ClonerEffector/Source/ClonerEffector/Private/Cloner/Extensions/CEClonerEmitterSpawnExtension.cpp b/Engine/Plugins/Experimental/ClonerEffector/Source/ClonerEffector/Private/Cloner/Extensions/CEClonerEmitterSpawnExtension.cpp
--- a/Engine/Plugins/Experimental/ClonerEffector/Source/ClonerEffector/Private/Cloner/Extensions/CEClonerEmitterSpawnExtension.cpp
+++ b/Engine/Plugins/Experimental/ClonerEffector/Source/ClonerEffector/Private/Cloner/Extensions/CEClonerEmitterSpawnExtension.cpp
@@ -2,11 +2,18 @@
 
 #include "Cloner/Extensions/CEClonerEmitterSpawnExtension.h"
 
-#include "NiagaraEmitter.h"
 #include "Cloner/CEClonerComponent.h"
+#include "NiagaraEmitter.h"
 #include "NiagaraSystem.h"
 #include "NiagaraUserRedirectionParameterStore.h"
 
+UCEClonerEmitterSpawnExtension::UCEClonerEmitterSpawnExtension()
+	: UCEClonerExtensionBase(
+		TEXT("Spawn")
+		, 0
+	)
+{}
+
 void UCEClonerEmitterSpawnExtension::SetSpawnLoopMode(ECEClonerSpawnLoopMode InMode)
 {
 	if (SpawnLoopMode == InMode)
diff --git a/Engine/Plugins/Experimental/ClonerEffector/Source/ClonerEffector/Private/Cloner/Extensions/CEClonerExtensionBase.cpp b/Engine/Plugins/Experimental/ClonerEffector/Source/ClonerEffector/Private/Cloner/Extensions/CEClonerExtensionBase.cpp
--- a/Engine/Plugins/Experimental/ClonerEffector/Source/ClonerEffector/Private/Cloner/Extensions/CEClonerExtensionBase.cpp
+++ b/Engine/Plugins/Experimental/ClonerEffector/Source/ClonerEffector/Private/Cloner/Extensions/CEClonerExtensionBase.cpp
@@ -13,6 +13,18 @@ FCEExtensionSection UCEClonerExtensionBase::GetExtensionSection() const
 }
 #endif
 
+UCEClonerExtensionBase::UCEClonerExtensionBase()
+	: UCEClonerExtensionBase(
+		NAME_None
+		, 0
+	)
+{}
+
+UCEClonerExtensionBase::UCEClonerExtensionBase(FName InExtensionName, int32 InExtensionPriority)
+	: ExtensionName(InExtensionName)
+	, ExtensionPriority(InExtensionPriority)
+{}
+
 UCEClonerComponent* UCEClonerExtensionBase::GetClonerComponent() const
 {
 	return GetTypedOuter<UCEClonerComponent>();
@@ -91,6 +103,21 @@ void UCEClonerExtensionBase::UpdateExtensionParameters()
 
 void UCEClonerExtensionBase::MarkExtensionDirty(bool bInUpdateCloner)
 {
+	if (!IsExtensionDirty())
+	{
+		// Notify other extensions once
+		if (const UCEClonerComponent* ClonerComponent = GetClonerComponent())
+		{
+			for (UCEClonerExtensionBase* ActiveExtension : ClonerComponent->GetActiveExtensions())
+			{
+				if (ActiveExtension)
+				{
+					ActiveExtension->OnExtensionDirtied(this);
+				}
+			}	
+		}
+	}
+
 	EnumAddFlags(ExtensionStatus, ECEClonerSystemStatus::ParametersDirty);
 
 	if (bInUpdateCloner)
diff --git a/Engine/Plugins/Experimental/ClonerEffector/Source/ClonerEffector/Private/Cloner/Extensions/CEClonerProgressExtension.cpp b/Engine/Plugins/Experimental/ClonerEffector/Source/ClonerEffector/Private/Cloner/Extensions/CEClonerProgressExtension.cpp
--- a/Engine/Plugins/Experimental/ClonerEffector/Source/ClonerEffector/Private/Cloner/Extensions/CEClonerProgressExtension.cpp
+++ b/Engine/Plugins/Experimental/ClonerEffector/Source/ClonerEffector/Private/Cloner/Extensions/CEClonerProgressExtension.cpp
@@ -4,6 +4,13 @@
 
 #include "Cloner/CEClonerComponent.h"
 
+UCEClonerProgressExtension::UCEClonerProgressExtension()
+	: UCEClonerExtensionBase(
+		TEXT("Progress")
+		, 0
+	)
+{}
+
 void UCEClonerProgressExtension::SetInvertProgress(bool bInInvertProgress)
 {
 	if (bInvertProgress == bInInvertProgress)
diff --git a/Engine/Plugins/Experimental/ClonerEffector/Source/ClonerEffector/Private/Cloner/Extensions/CEClonerRangeExtension.cpp b/Engine/Plugins/Experimental/ClonerEffector/Source/ClonerEffector/Private/Cloner/Extensions/CEClonerRangeExtension.cpp
--- a/Engine/Plugins/Experimental/ClonerEffector/Source/ClonerEffector/Private/Cloner/Extensions/CEClonerRangeExtension.cpp
+++ b/Engine/Plugins/Experimental/ClonerEffector/Source/ClonerEffector/Private/Cloner/Extensions/CEClonerRangeExtension.cpp
@@ -4,6 +4,13 @@
 
 #include "Cloner/CEClonerComponent.h"
 
+UCEClonerRangeExtension::UCEClonerRangeExtension()
+	: UCEClonerExtensionBase(
+		TEXT("Range")
+		, 0
+	)
+{}
+
 void UCEClonerRangeExtension::SetRangeEnabled(bool bInRangeEnabled)
 {
 	if (bRangeEnabled == bInRangeEnabled)
diff --git a/Engine/Plugins/Experimental/ClonerEffector/Source/ClonerEffector/Private/Cloner/Extensions/CEClonerStepExtension.cpp b/Engine/Plugins/Experimental/ClonerEffector/Source/ClonerEffector/Private/Cloner/Extensions/CEClonerStepExtension.cpp
--- a/Engine/Plugins/Experimental/ClonerEffector/Source/ClonerEffector/Private/Cloner/Extensions/CEClonerStepExtension.cpp
+++ b/Engine/Plugins/Experimental/ClonerEffector/Source/ClonerEffector/Private/Cloner/Extensions/CEClonerStepExtension.cpp
@@ -4,6 +4,13 @@
 
 #include "Cloner/CEClonerComponent.h"
 
+UCEClonerStepExtension::UCEClonerStepExtension()
+	: UCEClonerExtensionBase(
+		TEXT("Step")
+		, 0
+	)
+{}
+
 void UCEClonerStepExtension::SetDeltaStepEnabled(bool bInEnabled)
 {
 	if (bDeltaStepEnabled == bInEnabled)
diff --git a/Engine/Plugins/Experimental/ClonerEffector/Source/ClonerEffector/Private/Cloner/Extensions/CEClonerTextureExtension.cpp b/Engine/Plugins/Experimental/ClonerEffector/Source/ClonerEffector/Private/Cloner/Extensions/CEClonerTextureExtension.cpp
new file mode 100644
--- /dev/null
+++ b/Engine/Plugins/Experimental/ClonerEffector/Source/ClonerEffector/Private/Cloner/Extensions/CEClonerTextureExtension.cpp
@@ -0,0 +1,239 @@
+// Copyright Epic Games, Inc. All Rights Reserved.
+
+#include "Cloner/Extensions/CEClonerTextureExtension.h"
+
+#include "Cloner/CEClonerComponent.h"
+#include "Cloner/Extensions/CEClonerConstraintExtension.h"
+#include "Cloner/Extensions/CEClonerDisplacementExtension.h"
+#include "Cloner/Layouts/CEClonerGridLayout.h"
+#include "Cloner/Layouts/CEClonerLayoutBase.h"
+#include "NiagaraMeshRendererProperties.h"
+
+#if WITH_EDITOR
+const TCEPropertyChangeDispatcher<UCEClonerTextureExtension> UCEClonerTextureExtension::PropertyChangeDispatcher =
+{
+	{ GET_MEMBER_NAME_CHECKED(UCEClonerTextureExtension, bTextureEnabled), &UCEClonerTextureExtension::OnExtensionPropertyChanged },
+	{ GET_MEMBER_NAME_CHECKED(UCEClonerTextureExtension, TextureProvider), &UCEClonerTextureExtension::OnExtensionPropertyChanged },
+	{ GET_MEMBER_NAME_CHECKED(UCEClonerTextureExtension, CustomTextureAsset), &UCEClonerTextureExtension::OnExtensionPropertyChanged },
+	{ GET_MEMBER_NAME_CHECKED(UCEClonerTextureExtension, TextureUVProvider), &UCEClonerTextureExtension::OnExtensionPropertyChanged },
+	{ GET_MEMBER_NAME_CHECKED(UCEClonerTextureExtension, CustomTextureUVPlane), &UCEClonerTextureExtension::OnExtensionPropertyChanged },
+	{ GET_MEMBER_NAME_CHECKED(UCEClonerTextureExtension, CustomTextureUVOffset), &UCEClonerTextureExtension::OnExtensionPropertyChanged },
+	{ GET_MEMBER_NAME_CHECKED(UCEClonerTextureExtension, CustomTextureUVRotation), &UCEClonerTextureExtension::OnExtensionPropertyChanged },
+	{ GET_MEMBER_NAME_CHECKED(UCEClonerTextureExtension, CustomTextureUVScale), &UCEClonerTextureExtension::OnExtensionPropertyChanged },
+	{ GET_MEMBER_NAME_CHECKED(UCEClonerTextureExtension, bCustomTextureUVClamp), &UCEClonerTextureExtension::OnExtensionPropertyChanged },
+};
+
+void UCEClonerTextureExtension::PostEditChangeProperty(FPropertyChangedEvent& InPropertyChangedEvent)
+{
+	Super::PostEditChangeProperty(InPropertyChangedEvent);
+
+	PropertyChangeDispatcher.OnPropertyChanged(this, InPropertyChangedEvent);
+}
+#endif
+
+UCEClonerTextureExtension::UCEClonerTextureExtension()
+	: UCEClonerExtensionBase(
+		TEXT("Texture")
+		, 0
+	)
+{}
+
+void UCEClonerTextureExtension::SetTextureEnabled(bool bInEnabled)
+{
+	if (bTextureEnabled == bInEnabled)
+	{
+		return;
+	}
+
+	bTextureEnabled = bInEnabled;
+	MarkExtensionDirty();
+}
+
+void UCEClonerTextureExtension::SetTextureProvider(ECEClonerTextureProvider InProvider)
+{
+	if (TextureProvider == InProvider)
+	{
+		return;
+	}
+
+	TextureProvider = InProvider;
+	MarkExtensionDirty();
+}
+
+void UCEClonerTextureExtension::SetCustomTextureAsset(UTexture* InTexture)
+{
+	if (CustomTextureAsset == InTexture)
+	{
+		return;
+	}
+
+	CustomTextureAsset = InTexture;
+	MarkExtensionDirty();
+}
+
+void UCEClonerTextureExtension::SetTextureUVProvider(ECEClonerTextureProvider InProvider)
+{
+	if (TextureUVProvider == InProvider)
+	{
+		return;
+	}
+
+	TextureUVProvider = InProvider;
+	MarkExtensionDirty();
+}
+
+void UCEClonerTextureExtension::SetCustomTextureUVPlane(ECEClonerPlane InPlane)
+{
+	if (CustomTextureUVPlane == InPlane)
+	{
+		return;
+	}
+
+	CustomTextureUVPlane = InPlane;
+	MarkExtensionDirty();
+}
+
+void UCEClonerTextureExtension::SetCustomTextureUVOffset(const FVector2D& InOffset)
+{
+	if (CustomTextureUVOffset.Equals(InOffset))
+	{
+		return;
+	}
+
+	CustomTextureUVOffset = InOffset;
+	MarkExtensionDirty();
+}
+
+void UCEClonerTextureExtension::SetCustomTextureUVRotation(float InRotation)
+{
+	if (FMath::IsNearlyEqual(CustomTextureUVRotation, InRotation))
+	{
+		return;
+	}
+
+	CustomTextureUVRotation = InRotation;
+	MarkExtensionDirty();
+}
+
+void UCEClonerTextureExtension::SetCustomTextureUVScale(const FVector2D& InScale)
+{
+	if (CustomTextureUVScale.Equals(InScale))
+	{
+		return;
+	}
+
+	CustomTextureUVScale = InScale;
+	MarkExtensionDirty();
+}
+
+void UCEClonerTextureExtension::SetCustomTextureUVClamp(bool bInClamp)
+{
+	if (bCustomTextureUVClamp == bInClamp)
+	{
+		return;
+	}
+
+	bCustomTextureUVClamp = bInClamp;
+	MarkExtensionDirty();
+}
+
+void UCEClonerTextureExtension::OnExtensionParametersChanged(UCEClonerComponent* InComponent)
+{
+	Super::OnExtensionParametersChanged(InComponent);
+
+	const UCEClonerLayoutBase* ActiveLayout = InComponent->GetActiveLayout();
+
+	if (!ActiveLayout)
+	{
+		return;
+	}
+
+	UNiagaraMeshRendererProperties* MeshRenderer = ActiveLayout->GetMeshRenderer();
+
+	if (!MeshRenderer)
+	{
+		return;
+	}
+
+	FNiagaraRendererMaterialScalarParameter* TextureEnabledParameter = MeshRenderer->MaterialParameters.ScalarParameters.FindByPredicate(
+		[](const FNiagaraRendererMaterialScalarParameter& InMaterialParameter)
+		{
+			return InMaterialParameter.MaterialParameterName == TEXT("ClonerTextureEnabled");
+		}
+	);
+
+	if (!TextureEnabledParameter)
+	{
+		TextureEnabledParameter = &MeshRenderer->MaterialParameters.ScalarParameters.AddDefaulted_GetRef();
+		TextureEnabledParameter->MaterialParameterName = TEXT("ClonerTextureEnabled");
+	}
+
+	TextureEnabledParameter->Value = bTextureEnabled ? 1.0 : 0.0;
+
+	FNiagaraRendererMaterialTextureParameter* MaterialTextureParameter = MeshRenderer->MaterialParameters.TextureParameters.FindByPredicate(
+		[](const FNiagaraRendererMaterialTextureParameter& InMaterialParameter)
+		{
+			return InMaterialParameter.MaterialParameterName == TEXT("ClonerTexture");
+		}
+	);
+
+	if (!MaterialTextureParameter)
+	{
+		MaterialTextureParameter = &MeshRenderer->MaterialParameters.TextureParameters.AddDefaulted_GetRef();
+		MaterialTextureParameter->MaterialParameterName = TEXT("ClonerTexture");
+	}
+
+	if (TextureProvider == ECEClonerTextureProvider::Constraint)
+	{
+		if (const UCEClonerConstraintExtension* ConstraintExtension = InComponent->GetExtension<UCEClonerConstraintExtension>())
+		{
+			MaterialTextureParameter->Texture = ConstraintExtension->GetTextureAsset();
+		}
+	}
+	else if (TextureProvider == ECEClonerTextureProvider::Displacement)
+	{
+		if (const UCEClonerDisplacementExtension* DisplacementExtension = InComponent->GetExtension<UCEClonerDisplacementExtension>())
+		{
+			MaterialTextureParameter->Texture = DisplacementExtension->GetDisplacementTextureAsset();
+		}
+	}
+	else
+	{
+		MaterialTextureParameter->Texture = CustomTextureAsset;
+	}
+
+	FNiagaraUserRedirectionParameterStore& ExposedParameters = InComponent->GetOverrideParameters();
+
+	InComponent->SetVariableInt(TEXT("TextureUVIndex"), static_cast<int32>(TextureUVProvider));
+
+	static const FNiagaraVariable TexturePlaneVar(FNiagaraTypeDefinition(StaticEnum<ECEClonerPlane>()), TEXT("CustomTextureUVPlane"));
+	ExposedParameters.SetParameterValue<int32>(static_cast<int32>(CustomTextureUVPlane), TexturePlaneVar);
+	
+	InComponent->SetVariableVec2(TEXT("CustomTextureUVOffset"), CustomTextureUVOffset);
+
+	InComponent->SetFloatParameter(TEXT("CustomTextureUVRotation"), CustomTextureUVRotation);
+
+	InComponent->SetVariableVec2(TEXT("CustomTextureUVScale"), CustomTextureUVScale);
+
+	InComponent->SetVariableBool(TEXT("CustomTextureUVClamp"), bCustomTextureUVClamp);
+}
+
+bool UCEClonerTextureExtension::IsLayoutSupported(const UCEClonerLayoutBase* InLayout) const
+{
+	return InLayout->IsA<UCEClonerGridLayout>();
+}
+
+void UCEClonerTextureExtension::OnExtensionDirtied(const UCEClonerExtensionBase* InExtension)
+{
+	if (TextureProvider == ECEClonerTextureProvider::Displacement
+		|| TextureProvider == ECEClonerTextureProvider::Constraint
+		|| TextureUVProvider == ECEClonerTextureProvider::Displacement
+		|| TextureUVProvider == ECEClonerTextureProvider::Constraint)
+	{
+		if (InExtension->IsA<UCEClonerConstraintExtension>()
+			|| InExtension->IsA<UCEClonerDisplacementExtension>())
+		{
+			MarkExtensionDirty();
+		}
+	}
+}
diff --git a/Engine/Plugins/Experimental/ClonerEffector/Source/ClonerEffector/Public/Cloner/Extensions/CEClonerCollisionExtension.h b/Engine/Plugins/Experimental/ClonerEffector/Source/ClonerEffector/Public/Cloner/Extensions/CEClonerCollisionExtension.h
--- a/Engine/Plugins/Experimental/ClonerEffector/Source/ClonerEffector/Public/Cloner/Extensions/CEClonerCollisionExtension.h
+++ b/Engine/Plugins/Experimental/ClonerEffector/Source/ClonerEffector/Public/Cloner/Extensions/CEClonerCollisionExtension.h
@@ -11,18 +11,13 @@ class UMaterialInstanceDynamic;
 class UStaticMeshComponent;
 
 /** Extension dealing with collisions and physics related options */
-UCLASS(MinimalAPI, BlueprintType, Within=CEClonerComponent, AutoExpandCategories=("Collisions"), meta=(Section="Physics", Priority=3))
+UCLASS(MinimalAPI, BlueprintType, Within=CEClonerComponent, AutoExpandCategories=("Collisions"), meta=(Section="Physics", Priority=90))
 class UCEClonerCollisionExtension : public UCEClonerExtensionBase
 {
 	GENERATED_BODY()
 
 public:
-	UCEClonerCollisionExtension()
-		: UCEClonerExtensionBase(
-			TEXT("Collisions")
-			, 0
-		)
-	{}
+	UCEClonerCollisionExtension();
 
 	UFUNCTION(BlueprintCallable, Category="Cloner")
 	CLONEREFFECTOR_API void SetSurfaceCollisionEnabled(bool bInSurfaceCollisionEnabled);
diff --git a/Engine/Plugins/Experimental/ClonerEffector/Source/ClonerEffector/Public/Cloner/Extensions/CEClonerConstraintExtension.h b/Engine/Plugins/Experimental/ClonerEffector/Source/ClonerEffector/Public/Cloner/Extensions/CEClonerConstraintExtension.h
--- a/Engine/Plugins/Experimental/ClonerEffector/Source/ClonerEffector/Public/Cloner/Extensions/CEClonerConstraintExtension.h
+++ b/Engine/Plugins/Experimental/ClonerEffector/Source/ClonerEffector/Public/Cloner/Extensions/CEClonerConstraintExtension.h
@@ -7,18 +7,13 @@
 #include "CEClonerConstraintExtension.generated.h"
 
 /** Extension dealing with clone grid constraints */
-UCLASS(MinimalAPI, BlueprintType, Within=CEClonerComponent, meta=(Section="Cloner", Priority=0))
+UCLASS(MinimalAPI, BlueprintType, Within=CEClonerComponent, meta=(Section="Cloner", Priority=40))
 class UCEClonerConstraintExtension : public UCEClonerExtensionBase
 {
 	GENERATED_BODY()
 
 public:
-	UCEClonerConstraintExtension()
-		: UCEClonerExtensionBase(
-			TEXT("Constraint")
-			, 0
-		)
-	{}
+	UCEClonerConstraintExtension();
 
 	UFUNCTION(BlueprintCallable, Category="Cloner")
 	CLONEREFFECTOR_API void SetConstraint(ECEClonerGridConstraint InConstraint);
diff --git a/Engine/Plugins/Experimental/ClonerEffector/Source/ClonerEffector/Public/Cloner/Extensions/CEClonerDisplacementExtension.h b/Engine/Plugins/Experimental/ClonerEffector/Source/ClonerEffector/Public/Cloner/Extensions/CEClonerDisplacementExtension.h
--- a/Engine/Plugins/Experimental/ClonerEffector/Source/ClonerEffector/Public/Cloner/Extensions/CEClonerDisplacementExtension.h
+++ b/Engine/Plugins/Experimental/ClonerEffector/Source/ClonerEffector/Public/Cloner/Extensions/CEClonerDisplacementExtension.h
@@ -7,18 +7,13 @@
 #include "CEClonerDisplacementExtension.generated.h"
 
 /** Extension dealing with clone grid displacements */
-UCLASS(MinimalAPI, BlueprintType, Within=CEClonerComponent, meta=(Section="Cloner", Priority=0))
+UCLASS(MinimalAPI, BlueprintType, Within=CEClonerComponent, meta=(Section="Cloner", Priority=50))
 class UCEClonerDisplacementExtension : public UCEClonerExtensionBase
 {
 	GENERATED_BODY()
 
 public:
-	UCEClonerDisplacementExtension()
-		: UCEClonerExtensionBase(
-			TEXT("Displacement")
-			, 0
-		)
-	{}
+	UCEClonerDisplacementExtension();
 
 	UFUNCTION(BlueprintCallable, Category="Cloner")
 	CLONEREFFECTOR_API void SetDisplacementEnabled(bool bInEnabled);
diff --git a/Engine/Plugins/Experimental/ClonerEffector/Source/ClonerEffector/Public/Cloner/Extensions/CEClonerEffectorExtension.h b/Engine/Plugins/Experimental/ClonerEffector/Source/ClonerEffector/Public/Cloner/Extensions/CEClonerEffectorExtension.h
--- a/Engine/Plugins/Experimental/ClonerEffector/Source/ClonerEffector/Public/Cloner/Extensions/CEClonerEffectorExtension.h
+++ b/Engine/Plugins/Experimental/ClonerEffector/Source/ClonerEffector/Public/Cloner/Extensions/CEClonerEffectorExtension.h
@@ -12,7 +12,7 @@ class UCEEffectorComponent;
 struct FCEClonerEffectorDataInterfaces;
 
 /** Extension dealing with effectors options */
-UCLASS(MinimalAPI, BlueprintType, Within=CEClonerComponent, meta=(Section="Effector", Priority=1))
+UCLASS(MinimalAPI, BlueprintType, Within=CEClonerComponent, meta=(Section="Effector", Priority=30))
 class UCEClonerEffectorExtension : public UCEClonerExtensionBase
 {
 	GENERATED_BODY()
@@ -22,12 +22,7 @@ public:
 	static CLONEREFFECTOR_API FName GetEffectorActorsWeakName();
 #endif
 
-	UCEClonerEffectorExtension()
-		: UCEClonerExtensionBase(
-			TEXT("Effector")
-			, 0
-		)
-	{}
+	UCEClonerEffectorExtension();
 
 	/** Gets the number of effectors applied on this cloner */
 	UFUNCTION(BlueprintPure, Category="Cloner")
diff --git a/Engine/Plugins/Experimental/ClonerEffector/Source/ClonerEffector/Public/Cloner/Extensions/CEClonerEmitterSpawnExtension.h b/Engine/Plugins/Experimental/ClonerEffector/Source/ClonerEffector/Public/Cloner/Extensions/CEClonerEmitterSpawnExtension.h
--- a/Engine/Plugins/Experimental/ClonerEffector/Source/ClonerEffector/Public/Cloner/Extensions/CEClonerEmitterSpawnExtension.h
+++ b/Engine/Plugins/Experimental/ClonerEffector/Source/ClonerEffector/Public/Cloner/Extensions/CEClonerEmitterSpawnExtension.h
@@ -10,17 +10,13 @@
 class UNiagaraDataInterfaceCurve;
 
 /** Extension dealing with clones spawning options */
-UCLASS(MinimalAPI, BlueprintType, Within=CEClonerComponent, meta=(Section="Emission", Priority=2))
+UCLASS(MinimalAPI, BlueprintType, Within=CEClonerComponent, meta=(Section="Emission", Priority=70))
 class UCEClonerEmitterSpawnExtension : public UCEClonerExtensionBase
 {
 	GENERATED_BODY()
+
 public:
-	UCEClonerEmitterSpawnExtension()
-		: UCEClonerExtensionBase(
-			TEXT("Spawn")
-			, 0
-		)
-	{}
+	UCEClonerEmitterSpawnExtension();
 
 	UFUNCTION(BlueprintCallable, Category="Cloner")
 	CLONEREFFECTOR_API void SetSpawnLoopMode(ECEClonerSpawnLoopMode InMode);
diff --git a/Engine/Plugins/Experimental/ClonerEffector/Source/ClonerEffector/Public/Cloner/Extensions/CEClonerExtensionBase.h b/Engine/Plugins/Experimental/ClonerEffector/Source/ClonerEffector/Public/Cloner/Extensions/CEClonerExtensionBase.h
--- a/Engine/Plugins/Experimental/ClonerEffector/Source/ClonerEffector/Public/Cloner/Extensions/CEClonerExtensionBase.h
+++ b/Engine/Plugins/Experimental/ClonerEffector/Source/ClonerEffector/Public/Cloner/Extensions/CEClonerExtensionBase.h
@@ -16,19 +16,9 @@ class UCEClonerExtensionBase : public UObject
 	GENERATED_BODY()
 
 public:
-	UCEClonerExtensionBase()
-		: UCEClonerExtensionBase(
-			NAME_None
-			, 0
-		)
-	{}
-
-	UCEClonerExtensionBase(
-		FName InExtensionName
-		, int32 InExtensionPriority)
-		: ExtensionName(InExtensionName)
-		, ExtensionPriority(InExtensionPriority)
-	{}
+	UCEClonerExtensionBase();
+
+	UCEClonerExtensionBase(FName InExtensionName, int32 InExtensionPriority);
 
 	UFUNCTION(BlueprintPure, Category="Cloner")
 	FName GetExtensionName() const
@@ -84,6 +74,9 @@ public:
 		return true;
 	}
 
+	/** Called when an active extension is dirtied to allow other extensions to react accordingly */
+	virtual void OnExtensionDirtied(const UCEClonerExtensionBase* InExtension) {}
+
 protected:
 	//~ Begin UObject
 	virtual void PostEditImport() override;
diff --git a/Engine/Plugins/Experimental/ClonerEffector/Source/ClonerEffector/Public/Cloner/Extensions/CEClonerLifetimeExtension.h b/Engine/Plugins/Experimental/ClonerEffector/Source/ClonerEffector/Public/Cloner/Extensions/CEClonerLifetimeExtension.h
--- a/Engine/Plugins/Experimental/ClonerEffector/Source/ClonerEffector/Public/Cloner/Extensions/CEClonerLifetimeExtension.h
+++ b/Engine/Plugins/Experimental/ClonerEffector/Source/ClonerEffector/Public/Cloner/Extensions/CEClonerLifetimeExtension.h
@@ -11,7 +11,7 @@ class UCurveFloat;
 class UNiagaraDataInterfaceCurve;
 
 /** Extension dealing with clones lifetime options */
-UCLASS(MinimalAPI, BlueprintType, Within=CEClonerComponent, meta=(Section="Emission", Priority=3))
+UCLASS(MinimalAPI, BlueprintType, Within=CEClonerComponent, meta=(Section="Emission", Priority=80))
 class UCEClonerLifetimeExtension : public UCEClonerExtensionBase
 {
 	GENERATED_BODY()
diff --git a/Engine/Plugins/Experimental/ClonerEffector/Source/ClonerEffector/Public/Cloner/Extensions/CEClonerMeshRendererExtension.h b/Engine/Plugins/Experimental/ClonerEffector/Source/ClonerEffector/Public/Cloner/Extensions/CEClonerMeshRendererExtension.h
--- a/Engine/Plugins/Experimental/ClonerEffector/Source/ClonerEffector/Public/Cloner/Extensions/CEClonerMeshRendererExtension.h
+++ b/Engine/Plugins/Experimental/ClonerEffector/Source/ClonerEffector/Public/Cloner/Extensions/CEClonerMeshRendererExtension.h
@@ -9,7 +9,7 @@
 #include "CEClonerMeshRendererExtension.generated.h"
 
 /** Extension dealing with mesh rendering options */
-UCLASS(MinimalAPI, BlueprintType, Within=CEClonerComponent, meta=(Section="Rendering", Priority=4))
+UCLASS(MinimalAPI, BlueprintType, Within=CEClonerComponent, meta=(Section="Rendering", Priority=110))
 class UCEClonerMeshRendererExtension : public UCEClonerExtensionBase
 {
 	GENERATED_BODY()
diff --git a/Engine/Plugins/Experimental/ClonerEffector/Source/ClonerEffector/Public/Cloner/Extensions/CEClonerProgressExtension.h b/Engine/Plugins/Experimental/ClonerEffector/Source/ClonerEffector/Public/Cloner/Extensions/CEClonerProgressExtension.h
--- a/Engine/Plugins/Experimental/ClonerEffector/Source/ClonerEffector/Public/Cloner/Extensions/CEClonerProgressExtension.h
+++ b/Engine/Plugins/Experimental/ClonerEffector/Source/ClonerEffector/Public/Cloner/Extensions/CEClonerProgressExtension.h
@@ -7,18 +7,13 @@
 #include "CEClonerProgressExtension.generated.h"
 
 /** Extension dealing with clone progress options */
-UCLASS(MinimalAPI, BlueprintType, Within=CEClonerComponent, meta=(Section="Emission", Priority=4))
+UCLASS(MinimalAPI, BlueprintType, Within=CEClonerComponent, meta=(Section="Emission", Priority=100))
 class UCEClonerProgressExtension : public UCEClonerExtensionBase
 {
 	GENERATED_BODY()
 
 public:
-	UCEClonerProgressExtension()
-		: UCEClonerExtensionBase(
-			TEXT("Progress")
-			, 0
-		)
-	{}
+	UCEClonerProgressExtension();
 
 	UFUNCTION(BlueprintCallable, Category="Cloner")
 	CLONEREFFECTOR_API void SetInvertProgress(bool bInInvertProgress);
diff --git a/Engine/Plugins/Experimental/ClonerEffector/Source/ClonerEffector/Public/Cloner/Extensions/CEClonerRangeExtension.h b/Engine/Plugins/Experimental/ClonerEffector/Source/ClonerEffector/Public/Cloner/Extensions/CEClonerRangeExtension.h
--- a/Engine/Plugins/Experimental/ClonerEffector/Source/ClonerEffector/Public/Cloner/Extensions/CEClonerRangeExtension.h
+++ b/Engine/Plugins/Experimental/ClonerEffector/Source/ClonerEffector/Public/Cloner/Extensions/CEClonerRangeExtension.h
@@ -7,18 +7,13 @@
 #include "CEClonerRangeExtension.generated.h"
 
 /** Extension dealing with range options */
-UCLASS(MinimalAPI, BlueprintType, Within=CEClonerComponent, meta=(Section="Cloner", Priority=0))
+UCLASS(MinimalAPI, BlueprintType, Within=CEClonerComponent, meta=(Section="Cloner", Priority=10))
 class UCEClonerRangeExtension : public UCEClonerExtensionBase
 {
 	GENERATED_BODY()
 
 public:
-	UCEClonerRangeExtension()
-		: UCEClonerExtensionBase(
-			TEXT("Range")
-			, 0
-		)
-	{}
+	UCEClonerRangeExtension();
 
 	UFUNCTION(BlueprintCallable, Category="Cloner")
 	CLONEREFFECTOR_API void SetRangeEnabled(bool bInRangeEnabled);
@@ -134,8 +129,8 @@ protected:
 	UPROPERTY(EditInstanceOnly, Setter="SetRangeEnabled", Getter="GetRangeEnabled", DisplayName="Enabled", Category="Range")
 	bool bRangeEnabled = false;
 
-	/** Mirrors offset and rotation values */
-	UPROPERTY(EditInstanceOnly, Setter="SetRangeMirrored", Getter="GetRangeMirrored", DisplayName="Mirrored", Category="Range", meta=(EditCondition="bRangeEnabled", EditConditionHides))
+	/** Mirrors max offset and rotation values in min offset and rotation */
+	UPROPERTY(EditInstanceOnly, Setter="SetRangeMirrored", Getter="GetRangeMirrored", DisplayName="Mirror Min & Max", Category="Range", meta=(EditCondition="bRangeEnabled", EditConditionHides))
 	bool bRangeMirrored = true;
 
 	UPROPERTY(EditInstanceOnly, Setter, Getter, DisplayName="OffsetMin", Category="Range", meta=(EditCondition="bRangeEnabled && !bRangeMirrored", EditConditionHides))
diff --git a/Engine/Plugins/Experimental/ClonerEffector/Source/ClonerEffector/Public/Cloner/Extensions/CEClonerStepExtension.h b/Engine/Plugins/Experimental/ClonerEffector/Source/ClonerEffector/Public/Cloner/Extensions/CEClonerStepExtension.h
--- a/Engine/Plugins/Experimental/ClonerEffector/Source/ClonerEffector/Public/Cloner/Extensions/CEClonerStepExtension.h
+++ b/Engine/Plugins/Experimental/ClonerEffector/Source/ClonerEffector/Public/Cloner/Extensions/CEClonerStepExtension.h
@@ -7,18 +7,13 @@
 #include "CEClonerStepExtension.generated.h"
 
 /** Extension dealing with delta step accumulated options */
-UCLASS(MinimalAPI, BlueprintType, Within=CEClonerComponent, meta=(Section="Cloner", Priority=0))
+UCLASS(MinimalAPI, BlueprintType, Within=CEClonerComponent, meta=(Section="Cloner", Priority=20))
 class UCEClonerStepExtension : public UCEClonerExtensionBase
 {
 	GENERATED_BODY()
 
 public:
-	UCEClonerStepExtension()
-		: UCEClonerExtensionBase(
-			TEXT("Step")
-			, 0
-		)
-	{}
+	UCEClonerStepExtension();
 
 	UFUNCTION(BlueprintCallable, Category="Cloner")
 	CLONEREFFECTOR_API void SetDeltaStepEnabled(bool bInEnabled);
diff --git a/Engine/Plugins/Experimental/ClonerEffector/Source/ClonerEffector/Public/Cloner/Extensions/CEClonerTextureExtension.h b/Engine/Plugins/Experimental/ClonerEffector/Source/ClonerEffector/Public/Cloner/Extensions/CEClonerTextureExtension.h
new file mode 100644
--- /dev/null
+++ b/Engine/Plugins/Experimental/ClonerEffector/Source/ClonerEffector/Public/Cloner/Extensions/CEClonerTextureExtension.h
@@ -0,0 +1,154 @@
+// Copyright Epic Games, Inc. All Rights Reserved.
+
+#pragma once
+
+#include "CEClonerExtensionBase.h"
+#include "CEPropertyChangeDispatcher.h"
+#include "CEClonerTextureExtension.generated.h"
+
+UENUM()
+enum class ECEClonerTextureProvider : uint8
+{
+	/** Use texture from constraint extension */
+	Constraint,
+	/** Use texture from displacement extension */
+	Displacement,
+	/** Use custom texture */
+	Custom
+};
+
+/** Extension dealing with clone grid texture and UV for material */
+UCLASS(MinimalAPI, BlueprintType, Within=CEClonerComponent, meta=(Section="Cloner", Priority=60))
+class UCEClonerTextureExtension : public UCEClonerExtensionBase
+{
+	GENERATED_BODY()
+
+public:
+	UCEClonerTextureExtension();
+
+	UFUNCTION(BlueprintCallable, Category="Cloner")
+	void SetTextureEnabled(bool bInEnabled);
+	
+	UFUNCTION(BlueprintPure, Category="Cloner")
+	bool GetTextureEnabled() const
+	{
+		return bTextureEnabled;
+	}
+
+	UFUNCTION(BlueprintCallable, Category="Cloner")
+	void SetTextureProvider(ECEClonerTextureProvider InProvider);
+
+	UFUNCTION(BlueprintPure, Category="Cloner")
+	ECEClonerTextureProvider GetTextureProvider() const
+	{
+		return TextureProvider;
+	}
+
+	UFUNCTION(BlueprintCallable, Category="Cloner")
+	void SetCustomTextureAsset(UTexture* InTexture);
+	
+	UFUNCTION(BlueprintPure, Category="Cloner")
+	UTexture* GetCustomTextureAsset() const
+	{
+		return CustomTextureAsset;
+	}
+
+	UFUNCTION(BlueprintCallable, Category="Cloner")
+	void SetTextureUVProvider(ECEClonerTextureProvider InProvider);
+
+	UFUNCTION(BlueprintPure, Category="Cloner")
+	ECEClonerTextureProvider GetTextureUVProvider() const
+	{
+		return TextureUVProvider;
+	}
+	
+	UFUNCTION(BlueprintCallable, Category="Cloner")
+	void SetCustomTextureUVPlane(ECEClonerPlane InPlane);
+
+	UFUNCTION(BlueprintPure, Category="Cloner")
+	ECEClonerPlane GetCustomTextureUVPlane() const
+	{
+		return CustomTextureUVPlane;
+	}
+
+	UFUNCTION(BlueprintCallable, Category="Cloner")
+	CLONEREFFECTOR_API void SetCustomTextureUVOffset(const FVector2D& InOffset);
+
+	UFUNCTION(BlueprintPure, Category="Cloner")
+	FVector2D GetCustomTextureUVOffset() const
+	{
+		return CustomTextureUVOffset;
+	}
+
+	UFUNCTION(BlueprintCallable, Category="Cloner")
+	CLONEREFFECTOR_API void SetCustomTextureUVRotation(float InRotation);
+
+	UFUNCTION(BlueprintPure, Category="Cloner")
+	float GetCustomTextureUVRotation() const
+	{
+		return CustomTextureUVRotation;
+	}
+	
+	UFUNCTION(BlueprintCallable, Category="Cloner")
+	CLONEREFFECTOR_API void SetCustomTextureUVScale(const FVector2D& InScale);
+
+	UFUNCTION(BlueprintPure, Category="Cloner")
+	FVector2D GetCustomTextureUVScale() const
+	{
+		return CustomTextureUVScale;
+	}
+	
+	UFUNCTION(BlueprintCallable, Category="Cloner")
+	CLONEREFFECTOR_API void SetCustomTextureUVClamp(bool bInClamp);
+
+	UFUNCTION(BlueprintPure, Category="Cloner")
+	bool GetCustomTextureUVClamp() const
+	{
+		return bCustomTextureUVClamp;
+	}
+	
+protected:
+	//~ Begin UObject
+#if WITH_EDITOR
+	virtual void PostEditChangeProperty(FPropertyChangedEvent& InPropertyChangedEvent) override;
+#endif
+	//~ End UObject
+
+	//~ Begin UCEClonerExtensionBase
+	virtual void OnExtensionParametersChanged(UCEClonerComponent* InComponent) override;
+	virtual bool IsLayoutSupported(const UCEClonerLayoutBase* InLayout) const override;
+	virtual void OnExtensionDirtied(const UCEClonerExtensionBase* InExtension) override;
+	//~ End UCEClonerExtensionBase
+
+	UPROPERTY(EditInstanceOnly, Setter="SetTextureEnabled", Getter="GetTextureEnabled", DisplayName="Enabled", Category="Texture")
+	bool bTextureEnabled = false;
+
+	UPROPERTY(EditInstanceOnly, Setter, Getter, DisplayName="Texture Provider", Category="Texture", meta=(EditCondition="bTextureEnabled", EditConditionHides))
+	ECEClonerTextureProvider TextureProvider = ECEClonerTextureProvider::Constraint;
+
+	UPROPERTY(EditInstanceOnly, Setter, Getter, DisplayName="Texture Asset", Category="Texture", meta=(EditCondition="bTextureEnabled && TextureProvider == ECEClonerTextureProvider::Custom", EditConditionHides))
+	TObjectPtr<UTexture> CustomTextureAsset;
+	
+	UPROPERTY(EditInstanceOnly, Setter, Getter, DisplayName="Texture UV Provider", Category="Texture", meta=(EditCondition="bTextureEnabled", EditConditionHides))
+	ECEClonerTextureProvider TextureUVProvider = ECEClonerTextureProvider::Constraint;
+
+	UPROPERTY(EditInstanceOnly, Setter, Getter, DisplayName="Texture UV Plane", Category="Texture", meta=(InvalidEnumValues="Custom", EditCondition="bTextureEnabled && TextureUVProvider == ECEClonerTextureProvider::Custom", EditConditionHides))
+	ECEClonerPlane CustomTextureUVPlane = ECEClonerPlane::XY;
+
+	UPROPERTY(EditInstanceOnly, Setter, Getter, DisplayName="Texture UV Offset", Category="Texture", meta=(EditCondition="bTextureEnabled && TextureUVProvider == ECEClonerTextureProvider::Custom", EditConditionHides))
+	FVector2D CustomTextureUVOffset = FVector2D::ZeroVector;
+
+	UPROPERTY(EditInstanceOnly, Setter, Getter, DisplayName="Texture UV Rotation", Category="Texture", meta=(EditCondition="bTextureEnabled && TextureUVProvider == ECEClonerTextureProvider::Custom", EditConditionHides))
+	float CustomTextureUVRotation = 0.f;
+
+	UPROPERTY(EditInstanceOnly, Setter, Getter, DisplayName="Texture UV Scale", Category="Texture", meta=(Delta="0.01", MotionDesignVectorWidget, AllowPreserveRatio="XY", EditCondition="bTextureEnabled && TextureUVProvider == ECEClonerTextureProvider::Custom", EditConditionHides))
+	FVector2D CustomTextureUVScale = FVector2D(1.0);
+
+	UPROPERTY(EditInstanceOnly, Setter="SetCustomTextureUVClamp", Getter="GetCustomTextureUVClamp", DisplayName="Texture UV Clamp", Category="Texture", meta=(EditCondition="bTextureEnabled && TextureUVProvider == ECEClonerTextureProvider::Custom", EditConditionHides))
+	bool bCustomTextureUVClamp = false;
+
+private:
+#if WITH_EDITOR
+	static const TCEPropertyChangeDispatcher<UCEClonerTextureExtension> PropertyChangeDispatcher;
+#endif
+};
\ No newline at end of file
