MotionDesign : ClonerEffector
- Fix crash due to invalid FGUID binding during add cloner track utility action

diff --git a/Engine/Plugins/Experimental/ClonerEffector/Source/ClonerEffectorEditor/Private/Cloner/Sequencer/MovieSceneClonerTrackEditor.cpp b/Engine/Plugins/Experimental/ClonerEffector/Source/ClonerEffectorEditor/Private/Cloner/Sequencer/MovieSceneClonerTrackEditor.cpp
--- a/Engine/Plugins/Experimental/ClonerEffector/Source/ClonerEffectorEditor/Private/Cloner/Sequencer/MovieSceneClonerTrackEditor.cpp
+++ b/Engine/Plugins/Experimental/ClonerEffector/Source/ClonerEffectorEditor/Private/Cloner/Sequencer/MovieSceneClonerTrackEditor.cpp
@@ -67,7 +67,7 @@ void FMovieSceneClonerTrackEditor::ExecuteAddTrack(const TSet<UCEClonerComponent
 		}
 
 		const FGuid ComponentBinding = SequencerPtr->GetHandleToObject(Cloner);
-		if (FocusedMovieScene->FindSpawnable(ComponentBinding))
+		if (!ComponentBinding.IsValid() || FocusedMovieScene->FindSpawnable(ComponentBinding))
 		{
 			// we only want to add tracks for possessables
 			return;
