MotionDesign : ClonerEffector
- Updated cloner component customization to show properties from parent class (UNiagaraComponent and above) in a separate category

diff --git a/Engine/Plugins/Experimental/ClonerEffector/Source/ClonerEffector/Public/Cloner/CEClonerComponent.h b/Engine/Plugins/Experimental/ClonerEffector/Source/ClonerEffector/Public/Cloner/CEClonerComponent.h
--- a/Engine/Plugins/Experimental/ClonerEffector/Source/ClonerEffector/Public/Cloner/CEClonerComponent.h
+++ b/Engine/Plugins/Experimental/ClonerEffector/Source/ClonerEffector/Public/Cloner/CEClonerComponent.h
@@ -20,7 +20,7 @@ UCLASS(MinimalAPI
 	, BlueprintType
 	, DisplayName = "Motion Design Cloner Component"
 	, AutoExpandCategories=(Cloner, Layout)
-	, HideCategories=(Niagara, Activation, Lighting, Attachment, Randomness, Parameters, Warmup, Compilation, Navigation, Tags, LOD, TextureStreaming, Mobile, RayTracing, AssetUserData, Cooking, HLOD, Rendering))
+	, HideCategories=(Transform, Niagara, Activation, Attachment, Randomness, Parameters))
 class UCEClonerComponent : public UNiagaraComponent
 {
 	GENERATED_BODY()
diff --git a/Engine/Plugins/Experimental/ClonerEffector/Source/ClonerEffectorEditor/Private/Cloner/Customizations/CEEditorClonerComponentDetailCustomization.cpp b/Engine/Plugins/Experimental/ClonerEffector/Source/ClonerEffectorEditor/Private/Cloner/Customizations/CEEditorClonerComponentDetailCustomization.cpp
--- a/Engine/Plugins/Experimental/ClonerEffector/Source/ClonerEffectorEditor/Private/Cloner/Customizations/CEEditorClonerComponentDetailCustomization.cpp
+++ b/Engine/Plugins/Experimental/ClonerEffector/Source/ClonerEffectorEditor/Private/Cloner/Customizations/CEEditorClonerComponentDetailCustomization.cpp
@@ -9,6 +9,7 @@
 #include "DetailLayoutBuilder.h"
 #include "DetailWidgetRow.h"
 #include "IDetailChildrenBuilder.h"
+#include "IDetailGroup.h"
 #include "Input/Reply.h"
 #include "IPropertyUtilities.h"
 #include "Modules/ModuleManager.h"
@@ -30,6 +31,14 @@ void FCEEditorClonerComponentDetailCustomization::CustomizeDetails(IDetailLayout
 	// Remove niagara utilities
 	InDetailBuilder.HideCategory(TEXT("NiagaraComponent_Utilities"));
 
+	// Remove details sections by hiding specific categories
+	InDetailBuilder.HideCategory(TEXT("LOD"));	
+	InDetailBuilder.HideCategory(TEXT("HLOD"));	
+	InDetailBuilder.HideCategory(TEXT("Navigation"));	
+	InDetailBuilder.HideCategory(TEXT("AssetUserData"));	
+	InDetailBuilder.HideCategory(TEXT("Cooking"));	
+	InDetailBuilder.HideCategory(TEXT("Tags"));
+
 	TSharedRef<IPropertyUtilities> PropertyUtilities = InDetailBuilder.GetPropertyUtilities();
 	PropertyUtilitiesWeak = PropertyUtilities.ToWeakPtr();
 	ClonerComponentsWeak = InDetailBuilder.GetObjectsOfTypeBeingCustomized<UCEClonerComponent>();
@@ -202,6 +211,37 @@ void FCEEditorClonerComponentDetailCustomization::CustomizeDetails(IDetailLayout
 				];
 		}
 	}
+
+	// Show properties from parent class (UNiagaraComponent and above) in a separate category
+	{
+		IDetailCategoryBuilder& CategoryBuilder = InDetailBuilder.EditCategory(TEXT("NiagaraComponent"));
+		CategoryBuilder.SetSortOrder(10000);
+		TMap<FName, IDetailGroup*> NiagaraComponentGroups;
+		for (FProperty* Property : TFieldRange<FProperty>(UNiagaraComponent::StaticClass(), EFieldIterationFlags::IncludeSuper))
+		{
+			if (Property && Property->HasAnyPropertyFlags(CPF_Edit | CPF_EditConst))
+			{
+				TSharedRef<IPropertyHandle> NiagaraComponentPropertyHandle = InDetailBuilder.GetProperty(Property->GetFName(), Property->GetOwnerClass());
+				if (NiagaraComponentPropertyHandle->IsValidHandle())
+				{
+					const FName CategoryName = NiagaraComponentPropertyHandle->GetDefaultCategoryName();
+
+					IDetailGroup* DetailGroup = nullptr;
+					if (NiagaraComponentGroups.Contains(CategoryName))
+					{
+						DetailGroup = NiagaraComponentGroups.FindChecked(CategoryName);
+					}
+					else
+					{
+						DetailGroup = &CategoryBuilder.AddGroup(CategoryName, NiagaraComponentPropertyHandle->GetDefaultCategoryText());
+						NiagaraComponentGroups.Add(CategoryName, DetailGroup);
+					}
+
+					DetailGroup->AddPropertyRow(NiagaraComponentPropertyHandle);
+				}
+			}
+		}
+	}
 }
 
 void FCEEditorClonerComponentDetailCustomization::RemoveEmptySections()
