MotionDesign : ClonerEffector
- Moving cloner attachment tree in its own files to reduce code in cloner component
- Added warning bounds mismatch on load, cached mesh is different than actual mesh, could be due to asset missing
- Improved cloner load time by reusing cache as much as possible instead of regenerating attachment meshes, resave the asset in order to benefit from these improvements
- Instead of checking and updating the tree each tick, delegates are used to mark the tree dirty reducing attachments checks
- Removing ability to control tree interval updates, this will be done internally

diff --git a/Engine/Plugins/Experimental/Avalanche/Source/AvalancheEffectors/Private/AvalancheEffectorsModule.cpp b/Engine/Plugins/Experimental/Avalanche/Source/AvalancheEffectors/Private/AvalancheEffectorsModule.cpp
--- a/Engine/Plugins/Experimental/Avalanche/Source/AvalancheEffectors/Private/AvalancheEffectorsModule.cpp
+++ b/Engine/Plugins/Experimental/Avalanche/Source/AvalancheEffectors/Private/AvalancheEffectorsModule.cpp
@@ -2,120 +2,31 @@
 
 #include "AvalancheEffectorsModule.h"
 
-#include "AvaActorUtils.h"
-#include "AvaSceneTree.h"
-#include "IAvaSceneInterface.h"
+#include "AvalancheEffectorsSceneTreeResolver.h"
+#include "Engine/Level.h"
 #include "Engine/World.h"
-#include "GameFramework/Actor.h"
 #include "Modules/ModuleManager.h"
 #include "Subsystems/CEClonerSubsystem.h"
-
-#if WITH_EDITOR
-#include "IAvaOutliner.h"
-#include "AvaOutlinerDefines.h"
-#include "AvaOutlinerSubsystem.h"
-#include "Item/AvaOutlinerActor.h"
-#endif
+#include "Templates/SharedPointer.h"
 
 void FAvalancheEffectorsModule::StartupModule()
 {
-	UCEClonerSubsystem::OnSubsystemInitialized().AddRaw(this, &FAvalancheEffectorsModule::RegisterCustomActorResolver);
+	UCEClonerSubsystem::OnGetSceneTreeResolver().BindStatic(&FAvalancheEffectorsModule::CreateSceneTreeResolver);
 }
 
 void FAvalancheEffectorsModule::ShutdownModule()
 {
-	UCEClonerSubsystem::OnSubsystemInitialized().RemoveAll(this);
-
-	UnregisterCustomActorResolver();
-}
-
-void FAvalancheEffectorsModule::RegisterCustomActorResolver()
-{
-	if (UCEClonerSubsystem* Subsystem = UCEClonerSubsystem::Get())
-	{
-		Subsystem->RegisterCustomActorResolver(UCEClonerSubsystem::FOnGetOrderedActors::CreateRaw(this, &FAvalancheEffectorsModule::GetOrderedChildrenActors));
-	}
-}
-
-void FAvalancheEffectorsModule::UnregisterCustomActorResolver()
-{
-	if (UObjectInitialized())
-	{
-		if (UCEClonerSubsystem* Subsystem = UCEClonerSubsystem::Get())
-		{
-			Subsystem->UnregisterCustomActorResolver();
-		}
-	}
+	UCEClonerSubsystem::OnGetSceneTreeResolver().Unbind();
 }
 
-TArray<AActor*> FAvalancheEffectorsModule::GetOrderedChildrenActors(const AActor* InParentActor)
+TSharedPtr<ICEClonerSceneTreeCustomResolver> FAvalancheEffectorsModule::CreateSceneTreeResolver(ULevel* InLevel)
 {
-	TArray<AActor*> ChildrenActors;
-
-	if (!InParentActor)
-	{
-		return ChildrenActors;
-	}
-
-	UWorld* ClonerWorld = InParentActor->GetTypedOuter<UWorld>();
-	if (!ClonerWorld)
-	{
-		return ChildrenActors;
-	}
-
-#if WITH_EDITOR
-	if (const UAvaOutlinerSubsystem* const OutlinerSubsystem = ClonerWorld->GetSubsystem<UAvaOutlinerSubsystem>())
-	{
-		if (const TSharedPtr<IAvaOutliner> AvaOutliner = OutlinerSubsystem->GetOutliner())
-		{
-			const FAvaOutlinerItemPtr OutlinerClonerItem = AvaOutliner->FindItem(InParentActor);
-			if (OutlinerClonerItem.IsValid())
-			{
-				for (FAvaOutlinerItemPtr OutlinerChild : OutlinerClonerItem->GetChildren())
-				{
-					if (!OutlinerChild.IsValid() || !OutlinerChild->IsA<FAvaOutlinerActor>())
-					{
-						continue;
-					}
-
-					if (const FAvaOutlinerActor* OutlinerActor = OutlinerChild->CastTo<FAvaOutlinerActor>())
-					{
-						ChildrenActors.Add(OutlinerActor->GetActor());
-					}
-				}
-			}
-
-			return ChildrenActors;
-		}
-	}
-#endif
-
-	if (const IAvaSceneInterface* SceneInterface = FAvaActorUtils::GetSceneInterfaceFromActor(InParentActor))
-	{
-		TArray<AActor*> AttachedActors;
-		InParentActor->GetAttachedActors(AttachedActors, true, false);
-		ChildrenActors.SetNumUninitialized(AttachedActors.Num(), EAllowShrinking::No);
-		const FAvaSceneTree& SceneTree = SceneInterface->GetSceneTree();
-
-		for (AActor* ChildActor : AttachedActors)
-		{
-			if (const FAvaSceneTreeNode* ChildNode = SceneTree.FindTreeNode(FAvaSceneItem(ChildActor, ClonerWorld)))
-			{
-				const int32 ChildIndex = ChildNode->GetLocalIndex();
-
-				if (ChildrenActors.IsValidIndex(ChildIndex))
-				{
-					ChildrenActors[ChildIndex] = ChildActor;
-				}
-			}
-		}
-	}
-	else
+	if (IsValid(InLevel))
 	{
-		InParentActor->GetAttachedActors(ChildrenActors, true, false);
+		return MakeShared<FAvaEffectorsSceneTreeResolver>(InLevel);
 	}
 
-	return ChildrenActors;
+	return nullptr;
 }
 
 IMPLEMENT_MODULE(FAvalancheEffectorsModule, AvalancheEffectors)
diff --git a/Engine/Plugins/Experimental/Avalanche/Source/AvalancheEffectors/Private/AvalancheEffectorsModule.h b/Engine/Plugins/Experimental/Avalanche/Source/AvalancheEffectors/Private/AvalancheEffectorsModule.h
--- a/Engine/Plugins/Experimental/Avalanche/Source/AvalancheEffectors/Private/AvalancheEffectorsModule.h
+++ b/Engine/Plugins/Experimental/Avalanche/Source/AvalancheEffectors/Private/AvalancheEffectorsModule.h
@@ -4,8 +4,11 @@
 
 #include "Containers/Array.h"
 #include "Modules/ModuleInterface.h"
+#include "Templates/SharedPointerFwd.h"
 
 class AActor;
+class ICEClonerSceneTreeCustomResolver;
+class ULevel;
 
 class FAvalancheEffectorsModule : public IModuleInterface
 {
@@ -16,8 +19,5 @@ public:
 	//~ End IModuleInterface interface
 
 protected:
-	void RegisterCustomActorResolver();
-	void UnregisterCustomActorResolver();
-
-	TArray<AActor*> GetOrderedChildrenActors(const AActor* InParentActor);
+	static TSharedPtr<ICEClonerSceneTreeCustomResolver> CreateSceneTreeResolver(ULevel* InLevel);
 };
diff --git a/Engine/Plugins/Experimental/Avalanche/Source/AvalancheEffectors/Private/AvalancheEffectorsSceneTreeResolver.cpp b/Engine/Plugins/Experimental/Avalanche/Source/AvalancheEffectors/Private/AvalancheEffectorsSceneTreeResolver.cpp
new file mode 100644
--- /dev/null
+++ b/Engine/Plugins/Experimental/Avalanche/Source/AvalancheEffectors/Private/AvalancheEffectorsSceneTreeResolver.cpp
@@ -0,0 +1,138 @@
+// Copyright Epic Games, Inc. All Rights Reserved.
+
+#include "AvalancheEffectorsSceneTreeResolver.h"
+
+#include "AvaActorUtils.h"
+#include "AvaSceneSubsystem.h"
+#include "AvaSceneTree.h"
+#include "Engine/Level.h"
+#include "Engine/World.h"
+#include "GameFramework/Actor.h"
+#include "IAvaSceneInterface.h"
+
+#if WITH_EDITOR
+#include "AvaOutlinerSubsystem.h"
+#include "AvaOutlinerUtils.h"
+#include "IAvaOutliner.h"
+#endif
+
+FAvaEffectorsSceneTreeResolver::FAvaEffectorsSceneTreeResolver(ULevel* InLevel)
+	: LevelWeak(InLevel)
+{
+}
+
+void FAvaEffectorsSceneTreeResolver::Activate()
+{
+#if WITH_EDITOR
+	const ULevel* Level = LevelWeak.Get();
+
+	if (!Level)
+	{
+		return;
+	}
+
+	const UWorld* World = Level->GetWorld();
+
+	if (!World)
+	{
+		return;
+	}
+
+	if (UAvaOutlinerSubsystem* const OutlinerSubsystem = World->GetSubsystem<UAvaOutlinerSubsystem>())
+	{
+		UAvaOutlinerSubsystem::FActorHierarchyChanged& ActorHierarchyChanged = OutlinerSubsystem->OnActorHierarchyChanged();
+
+		ActorHierarchyChanged.RemoveAll(this);
+		ActorHierarchyChanged.AddSP(this, &FAvaEffectorsSceneTreeResolver::OnOutlinerHierarchyChanged);
+
+		if (const TSharedPtr<IAvaOutliner> Outliner = OutlinerSubsystem->GetOutliner())
+		{
+			IAvaOutliner::FOnOutlinerLoaded& OnOutlinerLoaded = Outliner->GetOnOutlinerLoaded();
+			OnOutlinerLoaded.RemoveAll(this);
+			OnOutlinerLoaded.AddSP(this, &FAvaEffectorsSceneTreeResolver::OnOutlinerLoaded);
+		}
+	}
+#endif
+}
+
+void FAvaEffectorsSceneTreeResolver::Deactivate()
+{
+#if WITH_EDITOR
+	const ULevel* Level = LevelWeak.Get();
+		
+	if (!Level)
+	{
+		return;
+	}
+
+	const UWorld* World = Level->GetWorld();
+
+	if (!World)
+	{
+		return;
+	}
+
+	if (UAvaOutlinerSubsystem* const OutlinerSubsystem = World->GetSubsystem<UAvaOutlinerSubsystem>())
+	{
+		OutlinerSubsystem->OnActorHierarchyChanged().RemoveAll(this);
+
+		if (const TSharedPtr<IAvaOutliner> Outliner = OutlinerSubsystem->GetOutliner())
+		{
+			Outliner->GetOnOutlinerLoaded().RemoveAll(this);
+		}
+	}
+#endif
+}
+
+bool FAvaEffectorsSceneTreeResolver::GetDirectChildrenActor(AActor* InActor, TArray<AActor*>& OutActors) const
+{
+	if (!IsValid(InActor))
+	{
+		return false;	
+	}
+	
+	bool bIsOutlinerAttachedActors = false;
+	
+#if WITH_EDITOR
+	if (const UWorld* const World = InActor->GetWorld())
+	{
+		if (const UAvaOutlinerSubsystem* const OutlinerSubsystem = World->GetSubsystem<UAvaOutlinerSubsystem>())
+		{
+			if (const TSharedPtr<IAvaOutliner> AvaOutliner = OutlinerSubsystem->GetOutliner())
+			{
+				OutActors = FAvaOutlinerUtils::EditorOutlinerChildActors(AvaOutliner, InActor);
+				bIsOutlinerAttachedActors = true;
+			}
+		}
+	}
+#endif
+
+	if (!bIsOutlinerAttachedActors)
+	{
+		if (const IAvaSceneInterface* SceneInterface = FAvaActorUtils::GetSceneInterfaceFromActor(InActor))
+		{
+			const FAvaSceneTree& SceneTree = SceneInterface->GetSceneTree();
+			OutActors = SceneTree.GetChildActors(InActor);
+			bIsOutlinerAttachedActors = true;
+		}
+	}
+
+	return bIsOutlinerAttachedActors;
+}
+
+FAvaEffectorsSceneTreeResolver::FOnActorHierarchyChanged::RegistrationType& FAvaEffectorsSceneTreeResolver::OnActorHierarchyChanged()
+{
+	return OnHierarchyChangedDelegate;
+}
+
+#if WITH_EDITOR
+void FAvaEffectorsSceneTreeResolver::OnOutlinerLoaded()
+{
+	OnHierarchyChangedDelegate.Broadcast(/** Actor */nullptr);
+}
+
+void FAvaEffectorsSceneTreeResolver::OnOutlinerHierarchyChanged(AActor* InActor, const AActor* InParent, EAvaOutlinerHierarchyChangeType InChange)
+{
+	OnHierarchyChangedDelegate.Broadcast(InActor);
+}
+#endif
diff --git a/Engine/Plugins/Experimental/Avalanche/Source/AvalancheEffectors/Private/AvalancheEffectorsSceneTreeResolver.h b/Engine/Plugins/Experimental/Avalanche/Source/AvalancheEffectors/Private/AvalancheEffectorsSceneTreeResolver.h
new file mode 100644
--- /dev/null
+++ b/Engine/Plugins/Experimental/Avalanche/Source/AvalancheEffectors/Private/AvalancheEffectorsSceneTreeResolver.h
@@ -0,0 +1,33 @@
+// Copyright Epic Games, Inc. All Rights Reserved.
+
+#pragma once
+
+#include "Cloner/Attachments/CEClonerSceneTreeCustomResolver.h"
+#include "Containers/Array.h"
+#include "UObject/WeakObjectPtr.h"
+
+class AActor;
+enum class EAvaOutlinerHierarchyChangeType : uint8;
+
+/** This allows cloner to retrieve hierarchy of the Motion Design outliner and react accordingly */
+class FAvaEffectorsSceneTreeResolver : public ICEClonerSceneTreeCustomResolver
+{
+public:
+	explicit FAvaEffectorsSceneTreeResolver(ULevel* InLevel);
+
+	//~ Begin ICEClonerSceneTreeCustomResolver
+	virtual void Activate() override;
+	virtual void Deactivate() override;
+	virtual bool GetDirectChildrenActor(AActor* InActor, TArray<AActor*>& OutActors) const override;
+	virtual FOnActorHierarchyChanged::RegistrationType& OnActorHierarchyChanged() override;
+	//~ End ICEClonerSceneTreeCustomResolver
+
+private:
+#if WITH_EDITOR
+	void OnOutlinerLoaded();
+	void OnOutlinerHierarchyChanged(AActor* InActor, const AActor* InParent, EAvaOutlinerHierarchyChangeType InChange);
+#endif
+
+	FOnActorHierarchyChanged OnHierarchyChangedDelegate;
+	TWeakObjectPtr<ULevel> LevelWeak;
+};
\ No newline at end of file
diff --git a/Engine/Plugins/Experimental/ClonerEffector/Source/ClonerEffector/Private/CEClonerEffectorShared.cpp b/Engine/Plugins/Experimental/ClonerEffector/Source/ClonerEffector/Private/CEClonerEffectorShared.cpp
--- a/Engine/Plugins/Experimental/ClonerEffector/Source/ClonerEffector/Private/CEClonerEffectorShared.cpp
+++ b/Engine/Plugins/Experimental/ClonerEffector/Source/ClonerEffector/Private/CEClonerEffectorShared.cpp
@@ -23,44 +23,6 @@
 
 #define LOCTEXT_NAMESPACE "CEClonerEffectorShared"
 
-#if WITH_EDITOR
-bool FCEClonerAttachmentItem::CheckBoundsChanged(bool bInUpdate)
-{
-	bool bChanged = false;
-
-	if (const AActor* Actor = ItemActor.Get())
-	{
-		FBox LocalBox(EForceInit::ForceInit);
-
-		TArray<UPrimitiveComponent*> PrimitiveComponents;
-		Actor->GetComponents(PrimitiveComponents);
-		for (const UPrimitiveComponent* Component : PrimitiveComponents)
-		{
-			if (Component && Component->IsRegistered())
-			{
-				FTransform ComponentToActorTransform = Component->GetComponentTransform().GetRelativeTransform(Actor->GetActorTransform());
-				FBoxSphereBounds ComponentBounds = Component->CalcBounds(ComponentToActorTransform);
-
-				LocalBox += ComponentBounds.GetBox();
-			}
-		}
-
-		const FVector ActorOrigin = LocalBox.GetCenter();
-		const FVector ActorExtent = LocalBox.GetExtent();
-
-		bChanged = !ActorOrigin.Equals(Origin) || !ActorExtent.Equals(Extent);
-
-		if (bInUpdate)
-		{
-			Origin = ActorOrigin;
-			Extent = ActorExtent;
-		}
-	}
-
-	return bChanged;
-}
-#endif
-
 void FCEClonerEffectorChannelData::Write(UNiagaraDataChannelWriter* InWriter) const
 {
 	if (!InWriter)
diff --git a/Engine/Plugins/Experimental/ClonerEffector/Source/ClonerEffector/Private/CEModule.cpp b/Engine/Plugins/Experimental/ClonerEffector/Source/ClonerEffector/Private/CEModule.cpp
--- a/Engine/Plugins/Experimental/ClonerEffector/Source/ClonerEffector/Private/CEModule.cpp
+++ b/Engine/Plugins/Experimental/ClonerEffector/Source/ClonerEffector/Private/CEModule.cpp
@@ -2,4 +2,8 @@
 
 #include "Modules/ModuleManager.h"
 
+#include "Cloner/Logs/CEClonerLogs.h"
+
+DEFINE_LOG_CATEGORY(LogCECloner);
+
 IMPLEMENT_MODULE(FDefaultModuleImpl, ClonerEffector)
\ No newline at end of file
diff --git a/Engine/Plugins/Experimental/ClonerEffector/Source/ClonerEffector/Private/Cloner/Attachments/CEClonerAttachmentItem.cpp b/Engine/Plugins/Experimental/ClonerEffector/Source/ClonerEffector/Private/Cloner/Attachments/CEClonerAttachmentItem.cpp
new file mode 100644
--- /dev/null
+++ b/Engine/Plugins/Experimental/ClonerEffector/Source/ClonerEffector/Private/Cloner/Attachments/CEClonerAttachmentItem.cpp
@@ -0,0 +1,134 @@
+// Copyright Epic Games, Inc. All Rights Reserved.
+
+#include "Cloner/Attachments/CEClonerAttachmentItem.h"
+
+#include "Components/PrimitiveComponent.h"
+#include "CEMeshBuilder.h"
+#include "GameFramework/Actor.h"
+#include "Settings/CEClonerEffectorSettings.h"
+#include "UDynamicMesh.h"
+#include "Utilities/CEClonerEffectorUtilities.h"
+
+#if WITH_EDITOR
+bool FCEClonerAttachmentItem::CheckBoundsChanged(bool bInUpdate)
+{
+	bool bChanged = false;
+
+	if (ItemActor.IsValid())
+	{
+		const FBox LocalBounds = GetAttachmentBounds();
+		const FVector ActorOrigin = LocalBounds.GetCenter();
+		const FVector ActorExtent = LocalBounds.GetExtent();
+
+		bChanged = !ActorOrigin.Equals(Origin) || !ActorExtent.Equals(Extent);
+
+		if (bInUpdate)
+		{
+			Origin = ActorOrigin;
+			Extent = ActorExtent;
+		}
+	}
+
+	return bChanged;
+}
+
+FBox FCEClonerAttachmentItem::GetBakedMeshBounds() const
+{
+	FBox Bounds(ForceInitToZero);
+	if (BakedMesh)
+	{
+		BakedMesh->ProcessMesh([&Bounds](const FDynamicMesh3& InMesh)
+		{
+			Bounds = static_cast<FBox>(InMesh.GetBounds(/** Parallel */true));
+		});
+	}
+	return Bounds;
+}
+
+FBox FCEClonerAttachmentItem::GetAttachmentBounds() const
+{
+	FBox Bounds(ForceInitToZero);
+	if (const AActor* Actor = ItemActor.Get())
+	{
+		TArray<UPrimitiveComponent*> PrimitiveComponents;
+		Actor->GetComponents(PrimitiveComponents, /** IncludeChildren */false);
+
+		for (const UPrimitiveComponent* Component : PrimitiveComponents)
+		{
+			if (Component
+				&& Component->IsRegistered()
+				&& !Component->IsVisualizationComponent())
+			{
+				FTransform ComponentToActorTransform = Component->GetComponentTransform().GetRelativeTransform(Actor->GetActorTransform());
+				FBoxSphereBounds ComponentBounds = Component->CalcBounds(ComponentToActorTransform);
+
+				Bounds += ComponentBounds.GetBox();
+			}
+		}
+	}
+	return Bounds;
+}
+
+bool FCEClonerAttachmentItem::CheckMaterialsChanged(bool bInUpdate, TArray<TWeakObjectPtr<UMaterialInterface>>* OutInvalidMaterials)
+{
+	using namespace UE::ClonerEffector::Utilities;
+
+	bool bMaterialChanged = false;
+
+	if (const AActor* Actor = ItemActor.Get())
+	{
+		TArray<UPrimitiveComponent*> PrimitiveComponents;
+		Actor->GetComponents(PrimitiveComponents, /** IncludeChildrenActors */false);
+
+		int32 MatIdx = 0;
+		TArray<TWeakObjectPtr<UMaterialInterface>> NewMaterials;
+		NewMaterials.Reserve(PrimitiveComponents.Num());
+		UMaterialInterface* DefaultMaterial = LoadObject<UMaterialInterface>(nullptr, UCEClonerEffectorSettings::DefaultMaterialPath);
+
+		if (OutInvalidMaterials)
+		{
+			OutInvalidMaterials->Empty();
+		}
+
+		for (UPrimitiveComponent* PrimitiveComponent : PrimitiveComponents)
+		{
+			if (!PrimitiveComponent || !FCEMeshBuilder::HasAnyGeometry(PrimitiveComponent))
+			{
+				continue;
+			}
+
+			for (int32 MatIndex = 0; MatIndex < PrimitiveComponent->GetNumMaterials(); MatIndex++)
+			{
+				UMaterialInterface* PreviousMaterial = PrimitiveComponent->GetMaterial(MatIndex);
+				UMaterialInterface* NewMaterial = PreviousMaterial;
+
+				if (FilterSupportedMaterial(NewMaterial, DefaultMaterial) && OutInvalidMaterials)
+				{
+					OutInvalidMaterials->Add(PreviousMaterial);
+				}
+
+				if (!BakedMaterials.IsValidIndex(MatIdx)
+					|| BakedMaterials[MatIdx] != NewMaterial)
+				{
+					bMaterialChanged = true;
+				}
+
+				NewMaterials.Add(NewMaterial);
+				MatIdx++;
+			}
+		}
+
+		if (bInUpdate)
+		{
+			if (NewMaterials.Num() != BakedMaterials.Num())
+			{
+				MeshStatus = ECEClonerAttachmentStatus::Outdated;
+			}
+
+			BakedMaterials = NewMaterials;
+		}
+	}
+
+	return bMaterialChanged;
+}
+#endif
diff --git a/Engine/Plugins/Experimental/ClonerEffector/Source/ClonerEffector/Private/Cloner/Attachments/CEClonerAttachmentTree.cpp b/Engine/Plugins/Experimental/ClonerEffector/Source/ClonerEffector/Private/Cloner/Attachments/CEClonerAttachmentTree.cpp
new file mode 100644
--- /dev/null
+++ b/Engine/Plugins/Experimental/ClonerEffector/Source/ClonerEffector/Private/Cloner/Attachments/CEClonerAttachmentTree.cpp
@@ -0,0 +1,486 @@
+// Copyright Epic Games, Inc. All Rights Reserved.
+
+#include "Cloner/Attachments/CEClonerAttachmentTree.h"
+
+#include "Cloner/Attachments/CEClonerAttachmentTreeBehavior.h"
+#include "Cloner/Logs/CEClonerLogs.h"
+#include "Components/PrimitiveComponent.h"
+#include "GameFramework/Actor.h"
+#include "Subsystems/CEClonerSubsystem.h"
+#include "Templates/SharedPointer.h"
+
+void FCEClonerAttachmentTree::SetAttachmentRoot(AActor* InActor)
+{
+	AttachmentRoot = InActor;
+	MarkAttachmentOutdated();
+}
+
+AActor* FCEClonerAttachmentTree::GetAttachmentRoot() const
+{
+	return AttachmentRoot;
+}
+
+void FCEClonerAttachmentTree::SetBehaviorImplementation(const TSharedRef<ICEClonerAttachmentTreeBehavior>& InImplementation)
+{
+	if (BehaviorImplementation.IsValid())
+	{
+		ForEachAttachment([this](AActor* InActor, FCEClonerAttachmentItem& InAttachment)
+		{
+			DetachItem(InActor);
+			return true;
+		});
+		
+		BehaviorImplementation->OnDeactivation(*this);
+	}
+
+	BehaviorImplementation = InImplementation;
+
+	if (BehaviorImplementation.IsValid())
+	{
+		BehaviorImplementation->OnActivation(*this);
+		MarkAttachmentOutdated();
+	}
+}
+
+void FCEClonerAttachmentTree::Cleanup()
+{
+	if (BehaviorImplementation.IsValid())
+	{
+		BehaviorImplementation->OnDeactivation(*this);
+		BehaviorImplementation.Reset();
+	}
+}
+
+void FCEClonerAttachmentTree::Reset()
+{
+	ItemAttachmentMap.Reset();
+	RootActors.Reset();
+	MergedBakedMeshes.Reset();
+	DirtyItemAttachments.Reset();
+	Status = ECEClonerAttachmentStatus::Outdated;
+}
+
+void FCEClonerAttachmentTree::MarkAttachmentOutdated()
+{
+	if (Status == ECEClonerAttachmentStatus::Updated)
+	{
+		Status = ECEClonerAttachmentStatus::Outdated;
+	}
+}
+
+void FCEClonerAttachmentTree::MarkCacheOutdated(AActor* InActor)
+{
+	if (BehaviorImplementation.IsValid())
+	{
+		BehaviorImplementation->InvalidateMesh(*this, InActor);
+	}
+}
+
+bool FCEClonerAttachmentTree::UpdateAttachments(bool bInReset)
+{
+	if (!IsValid(AttachmentRoot)
+		|| !BehaviorImplementation.IsValid()
+		|| Status == ECEClonerAttachmentStatus::Updating)
+	{
+		return false;
+	}
+
+	if (bInReset)
+	{
+		Reset();
+		Status = ECEClonerAttachmentStatus::Outdated;
+	}
+
+	if (Status != ECEClonerAttachmentStatus::Outdated)
+	{
+		return false;
+	}
+
+	UE_LOG(LogCECloner, Log, TEXT("%s : updating attachment tree"), *AttachmentRoot->GetActorNameOrLabel())
+
+	Status = ECEClonerAttachmentStatus::Updating;
+	return UpdateAttachmentsInternal();
+}
+
+AActor* FCEClonerAttachmentTree::FindRootActor(AActor* InActor)
+{
+	if (InActor == nullptr)
+	{
+		return nullptr;
+	}
+
+	if (const FCEClonerAttachmentItem* Item = ItemAttachmentMap.Find(InActor))
+	{
+		if (Item->bRootItem)
+		{
+			return Item->ItemActor.Get();
+		}
+
+		return FindRootActor(Item->ParentActor.Get());
+	}
+
+	return nullptr;
+}
+
+TArray<AActor*> FCEClonerAttachmentTree::GetRootActors() const
+{
+	TArray<AActor*> Actors;
+	Algo::Transform(RootActors, Actors, [](const TWeakObjectPtr<AActor>& InActorWeak)->AActor*
+	{
+		return InActorWeak.Get();
+	});
+	return Actors;
+}
+
+void FCEClonerAttachmentTree::ForEachAttachment(const TFunctionRef<bool(AActor*, FCEClonerAttachmentItem&)>& InFunctor)
+{
+	for (TMap<TWeakObjectPtr<AActor>, FCEClonerAttachmentItem>::TIterator It(ItemAttachmentMap); It; ++It)
+	{
+		if (AActor* Actor = It->Key.Get())
+		{
+			if (!InFunctor(Actor, It->Value))
+			{
+				return;
+			}
+		}
+	}
+}
+
+void FCEClonerAttachmentTree::DetachItem(AActor* InActor)
+{
+	if (FCEClonerAttachmentItem* Item = ItemAttachmentMap.Find(InActor))
+	{
+		DetachItem(*Item);
+	}
+}
+
+bool FCEClonerAttachmentTree::IsCacheAvailable(bool bInAllowInvalid) const
+{
+	// Check cached meshes are not async loading
+	for (const TObjectPtr<UStaticMesh>& BakedMesh : MergedBakedMeshes)
+	{
+		if (!IsValid(BakedMesh))
+		{
+			if (!bInAllowInvalid)
+			{
+				return false;
+			}
+
+			continue;
+		}
+
+		if (BakedMesh->HasAnyInternalFlags(EInternalObjectFlags_AsyncLoading))
+		{
+			return false;
+		}
+	}
+
+	return true;
+}
+
+void FCEClonerAttachmentTree::GetAttachments(AActor* InActor, TArray<const FCEClonerAttachmentItem*>& OutAttachmentItems, bool bInRecurse) const
+{
+	if (!InActor)
+	{
+		return;
+	}
+
+	const FCEClonerAttachmentItem* AttachmentItem = ItemAttachmentMap.Find(InActor);
+
+	if (!AttachmentItem)
+	{
+		return;
+	}
+
+	OutAttachmentItems.Add(AttachmentItem);
+
+	if (bInRecurse)
+	{
+		for (const TWeakObjectPtr<AActor>& ChildActor : AttachmentItem->ChildrenActors)
+		{
+			GetAttachments(ChildActor.Get(), OutAttachmentItems, bInRecurse);
+		}
+	}
+}
+
+FBox FCEClonerAttachmentTree::GetAttachmentBounds(AActor* InActor, bool bInIncludeChildren) const
+{
+	TArray<const FCEClonerAttachmentItem*> Attachments;
+	GetAttachments(InActor, Attachments, bInIncludeChildren);
+
+	FBox Bounds(ForceInitToZero);
+	const FTransform ActorTransform = InActor->GetActorTransform();
+	for (const FCEClonerAttachmentItem* Attachment : Attachments)
+	{
+		if (const AActor* Actor = Attachment->ItemActor.Get())
+		{
+			TArray<UPrimitiveComponent*> PrimitiveComponents;
+			Actor->GetComponents(PrimitiveComponents, /** IncludeChildren */false);
+
+			for (const UPrimitiveComponent* Component : PrimitiveComponents)
+			{
+				if (Component
+					&& Component->IsRegistered()
+#if WITH_EDITOR
+					&& !Component->IsVisualizationComponent()
+#endif
+					)
+				{
+					FTransform ComponentToActorTransform = Component->GetComponentTransform().GetRelativeTransform(ActorTransform);
+					FBoxSphereBounds ComponentBounds = Component->CalcBounds(ComponentToActorTransform);
+
+					Bounds += ComponentBounds.GetBox();
+					Bounds.IsValid = true;
+				}
+			}
+		}
+	}
+
+	return Bounds;
+}
+
+FBox FCEClonerAttachmentTree::GetCacheBounds(AActor* InActor) const
+{
+	FBox Bounds(ForceInitToZero);
+	const int32 Index = RootActors.Find(InActor);
+
+	if (Index == INDEX_NONE)
+	{
+		return Bounds;
+	}
+
+	if (!MergedBakedMeshes.IsValidIndex(Index) || !MergedBakedMeshes[Index])
+	{
+		return Bounds;
+	}
+
+	Bounds = MergedBakedMeshes[Index]->GetBounds().GetBox();
+	Bounds.IsValid = true;
+
+	return Bounds;
+}
+
+bool FCEClonerAttachmentTree::UpdateAttachmentsInternal()
+{
+	// Invalidate all, to see what is outdated and what is still invalid
+	for (TPair<TWeakObjectPtr<AActor>, FCEClonerAttachmentItem>& AttachmentPair : ItemAttachmentMap)
+	{
+		AttachmentPair.Value.Status = ECEClonerAttachmentStatus::Invalid;
+	}
+
+	// Update root attachment items
+	TArray<AActor*> RootChildren;
+	BehaviorImplementation->GetOrderedChildrenActors(*this, AttachmentRoot, RootChildren);
+
+	TArray<TObjectPtr<UStaticMesh>> NewCombinesMeshes;
+	TArray<TWeakObjectPtr<AActor>> NewRootActors;
+	NewRootActors.Reserve(RootChildren.Num());
+	NewCombinesMeshes.Reserve(RootChildren.Num());
+	for (int32 RootIdx = 0; RootIdx < RootChildren.Num(); RootIdx++)
+	{
+		AActor* RootChild = RootChildren[RootIdx];
+
+		if (!IsValid(RootChild))
+		{
+			continue;
+		}
+
+		const FTransform RootTransform = RootChild->GetActorTransform();
+		UpdateAttachment(RootChild, nullptr, RootTransform);
+
+		// Lets find the old root idx
+		const int32 OldIdx = RootActors.Find(RootChild);
+		TObjectPtr<UStaticMesh> CombineBakedMesh = nullptr;
+		if (MergedBakedMeshes.IsValidIndex(OldIdx))
+		{
+			CombineBakedMesh = MergedBakedMeshes[OldIdx].Get();
+
+			// Did we rearrange stuff ?
+			if (RootIdx != OldIdx)
+			{
+				bItemAttachmentsDirty = true;
+			}
+		}
+		NewCombinesMeshes.Add(CombineBakedMesh);
+		NewRootActors.Add(RootChild);
+	}
+
+	// Did we remove any root actors ?
+	if (RootActors.Num() != NewRootActors.Num())
+	{
+		bItemAttachmentsDirty = true;
+	}
+
+	// Did we need to update meshes
+	TArray<TWeakObjectPtr<AActor>> ClonedActors;
+	ItemAttachmentMap.GenerateKeyArray(ClonedActors);
+	for (const TWeakObjectPtr<AActor>& ClonedActorWeak : ClonedActors)
+	{
+		FCEClonerAttachmentItem* ClonedItem = ItemAttachmentMap.Find(ClonedActorWeak);
+
+		if (!ClonedItem)
+		{
+			continue;
+		}
+
+		AActor* ClonedActor = ClonedActorWeak.Get();
+
+		if (ClonedItem->Status == ECEClonerAttachmentStatus::Invalid)
+		{
+			DetachItem(*ClonedItem);
+		}
+		else if (ClonedItem->Status == ECEClonerAttachmentStatus::Outdated)
+		{
+			if (ClonedItem->MeshStatus == ECEClonerAttachmentStatus::Outdated)
+			{
+				DirtyItemAttachments.Add(ClonedItem->ItemActor);
+				BehaviorImplementation->InvalidateMesh(*this, ClonedActor);
+			}
+
+			bItemAttachmentsDirty = true;
+			ClonedItem->Status = ECEClonerAttachmentStatus::Updated;
+		}
+	}
+
+	// Did we remove an attachment ?
+	if (ClonedActors.Num() != ItemAttachmentMap.Num())
+	{
+		bItemAttachmentsDirty = true;
+	}
+
+	if (!DirtyItemAttachments.IsEmpty())
+	{
+		bItemAttachmentsDirty = true;
+	}
+
+	RootActors = NewRootActors;
+	MergedBakedMeshes = NewCombinesMeshes;
+	Status = ECEClonerAttachmentStatus::Updated;
+
+	return bItemAttachmentsDirty;
+}
+
+void FCEClonerAttachmentTree::UpdateAttachment(AActor* InActor, AActor* InParent, const FTransform& InActorTransform)
+{
+	if (!IsValid(InActor))
+	{
+		return;
+	}
+
+	TArray<AActor*> ChildrenActors;
+	BehaviorImplementation->GetOrderedChildrenActors(*this, InActor, ChildrenActors);
+
+	const FTransform ActorTransform = InActor->GetActorTransform().GetRelativeTransform(InActorTransform);
+
+	FCEClonerAttachmentItem* AttachmentItem = ItemAttachmentMap.Find(InActor);
+
+	if (!AttachmentItem)
+	{
+		AttachmentItem = &ItemAttachmentMap.Add(InActor, FCEClonerAttachmentItem());
+		AttachmentItem->ItemActor = InActor;
+		AttachmentItem->MeshStatus = ECEClonerAttachmentStatus::Outdated;
+		AttachmentItem->Status = ECEClonerAttachmentStatus::Outdated;
+	}
+	else
+	{
+		AttachmentItem->Status = ECEClonerAttachmentStatus::Updated;
+	}
+
+	if (AttachmentItem)
+	{
+		// Check Root is the same
+		const bool bIsRoot = InParent == nullptr;
+		if (AttachmentItem->bRootItem != bIsRoot)
+		{
+			BehaviorImplementation->InvalidateMesh(*this, InActor);
+			AttachmentItem->bRootItem = bIsRoot;
+			AttachmentItem->Status = ECEClonerAttachmentStatus::Outdated;
+		}
+
+		// Check parent is the same
+		if (AttachmentItem->ParentActor.Get() != InParent)
+		{
+			BehaviorImplementation->InvalidateMesh(*this, InParent);
+			BehaviorImplementation->InvalidateMesh(*this, AttachmentItem->ParentActor.Get());
+			AttachmentItem->ParentActor = InParent;
+			AttachmentItem->Status = ECEClonerAttachmentStatus::Outdated;
+		}
+
+		// Check transform is the same
+		if (!ActorTransform.Equals(AttachmentItem->ActorTransform))
+		{
+			// invalidate if not root, else change transform in mesh renderer
+			if (!bIsRoot)
+			{
+				BehaviorImplementation->InvalidateMesh(*this, InActor);
+			}
+			AttachmentItem->ActorTransform = ActorTransform;
+			AttachmentItem->Status = ECEClonerAttachmentStatus::Outdated;
+		}
+
+#if WITH_EDITOR
+		// Check bounds change
+		if (AttachmentItem->CheckBoundsChanged(/** CheckAndUpdate */true) && AttachmentItem->bSetupDone)
+		{
+			BehaviorImplementation->InvalidateMesh(*this, InActor);
+			AttachmentItem->Status = ECEClonerAttachmentStatus::Outdated;
+		}
+
+		// Check materials change
+		TArray<TWeakObjectPtr<UMaterialInterface>> UnsupportedMaterials;
+		if (AttachmentItem->CheckMaterialsChanged(/** Update */true, &UnsupportedMaterials))
+		{
+			// Show warning for unset materials
+			if (!UnsupportedMaterials.IsEmpty())
+			{
+				if (UCEClonerSubsystem* ClonerSubsystem = UCEClonerSubsystem::Get())
+				{
+					ClonerSubsystem->FireMaterialWarning(AttachmentRoot, InActor, UnsupportedMaterials);
+				}
+			}
+
+			if (AttachmentItem->bSetupDone)
+			{
+				AttachmentItem->Status = ECEClonerAttachmentStatus::Outdated;
+			}
+		}
+
+		// Check dynamic mesh needs update
+		if (!AttachmentItem->BakedMesh && AttachmentItem->bSetupDone)
+		{
+			AttachmentItem->MeshStatus = ECEClonerAttachmentStatus::Outdated;
+			AttachmentItem->Status = ECEClonerAttachmentStatus::Outdated;
+		}
+#endif
+
+		if (AttachmentItem->ChildrenActors.Num() != ChildrenActors.Num())
+		{
+			BehaviorImplementation->InvalidateMesh(*this, InActor);
+		}
+
+		AttachmentItem->ChildrenActors.Empty(ChildrenActors.Num());
+		for (AActor* ChildActor : ChildrenActors)
+		{
+			AttachmentItem->ChildrenActors.Add(ChildActor);
+			UpdateAttachment(ChildActor, InActor, InActorTransform);
+		}
+
+		if (!AttachmentItem->bSetupDone)
+		{
+			AttachmentItem->bSetupDone = true;
+			BehaviorImplementation->OnItemAttached(*this, *AttachmentItem);
+			OnItemAttachedDelegate.ExecuteIfBound(InActor, *AttachmentItem);
+		}
+	}
+}
+
+void FCEClonerAttachmentTree::DetachItem(FCEClonerAttachmentItem& InItem)
+{
+	AActor* Actor = InItem.ItemActor.Get(/** EvenIfPendingKill */true);
+	BehaviorImplementation->InvalidateMesh(*this, Actor);
+	BehaviorImplementation->OnItemDetached(*this, InItem);
+	OnItemDetachedDelegate.ExecuteIfBound(Actor, InItem);
+	ItemAttachmentMap.Remove(Actor);
+	bItemAttachmentsDirty = true;
+}
+
diff --git a/Engine/Plugins/Experimental/ClonerEffector/Source/ClonerEffector/Private/Cloner/Attachments/CEClonerAttachmentTreeBehavior.cpp b/Engine/Plugins/Experimental/ClonerEffector/Source/ClonerEffector/Private/Cloner/Attachments/CEClonerAttachmentTreeBehavior.cpp
new file mode 100644
--- /dev/null
+++ b/Engine/Plugins/Experimental/ClonerEffector/Source/ClonerEffector/Private/Cloner/Attachments/CEClonerAttachmentTreeBehavior.cpp
@@ -0,0 +1,539 @@
+// Copyright Epic Games, Inc. All Rights Reserved.
+
+#include "Cloner/Attachments/CEClonerAttachmentTreeBehavior.h"
+
+#include "CEMeshBuilder.h"
+#include "Cloner/CEClonerComponent.h"
+#include "Cloner/Attachments/CEClonerAttachmentTree.h"
+#include "Cloner/Attachments/CEClonerSceneTreeCustomResolver.h"
+#include "Cloner/Logs/CEClonerLogs.h"
+#include "Components/DynamicMeshComponent.h"
+#include "Engine/Engine.h"
+#include "Engine/World.h"
+#include "GameFramework/Actor.h"
+#include "Materials/Material.h"
+#include "Subsystems/CEClonerSubsystem.h"
+#include "Utilities/CEClonerEffectorUtilities.h"
+
+#define LOCTEXT_NAMESPACE "CEClonerAttachmentImplementation"
+
+void FCEClonerAttachmentGroupBehavior::OnActivation(FCEClonerAttachmentTree& InTree)
+{
+	const AActor* AttachmentRoot = InTree.GetAttachmentRoot();
+
+	if (AttachmentRoot)
+	{
+		if (const UWorld* World = AttachmentRoot->GetWorld())
+		{
+			WorldActorDestroyedDelegate = World->AddOnActorDestroyedHandler(FOnActorDestroyed::FDelegate::CreateSPLambda(this, [this, &InTree](AActor* InActor)
+			{
+				OnWorldActorDestroyed(InTree, InActor);
+			}));
+		}
+	}
+
+#if WITH_EDITOR
+	USceneComponent::MarkRenderStateDirtyEvent.AddSPLambda(this, [this, &InTree](UActorComponent& InComponent)
+	{
+		OnRenderStateDirty(InTree, InComponent);
+	});
+
+	FCoreUObjectDelegates::OnObjectPropertyChanged.AddSPLambda(this, [this, &InTree](UObject* InObject, FPropertyChangedEvent& InEvent)
+	{
+		OnCheckMaterialChanged(InTree, InObject);
+	});
+
+	UMaterial::OnMaterialCompilationFinished().AddSPLambda(this, [this, &InTree](UMaterialInterface* InMaterial)
+	{
+		OnCheckMaterialChanged(InTree, InMaterial);
+	});
+
+	UCEClonerSubsystem* ClonerSubsystem = UCEClonerSubsystem::Get();
+	if (!!AttachmentRoot && !!ClonerSubsystem)
+	{
+		if (const TSharedPtr<ICEClonerSceneTreeCustomResolver> CustomSceneTreeResolver = ClonerSubsystem->FindCustomLevelSceneTreeResolver(AttachmentRoot->GetLevel()))
+		{
+			CustomSceneTreeResolver->OnActorHierarchyChanged().AddSPLambda(this, [this, &InTree](AActor* InActor)
+			{
+				OnLevelHierarchyChanged(InTree, InActor);
+			});
+		}
+	}
+
+	if (GEngine)
+	{
+		GEngine->OnLevelActorAttached().AddSPLambda(this, [this, &InTree](AActor* InActor, const AActor* InParent)
+		{
+			OnLevelHierarchyChanged(InTree, InActor);
+		});
+	
+		GEngine->OnLevelActorDetached().AddSPLambda(this, [this, &InTree](AActor* InActor, const AActor* InParent)
+		{
+			OnLevelHierarchyChanged(InTree, InActor);
+		});
+	}
+#endif
+}
+
+void FCEClonerAttachmentGroupBehavior::OnDeactivation(FCEClonerAttachmentTree& InTree)
+{
+	const AActor* AttachmentRoot = InTree.GetAttachmentRoot();
+
+	if (AttachmentRoot)
+	{
+		if (const UWorld* World = AttachmentRoot->GetWorld())
+		{
+			World->RemoveOnActorDestroyedHandler(WorldActorDestroyedDelegate);
+		}
+	}
+
+#if WITH_EDITOR
+	USceneComponent::MarkRenderStateDirtyEvent.RemoveAll(this);
+
+	FCoreUObjectDelegates::OnObjectPropertyChanged.RemoveAll(this);
+
+	UMaterial::OnMaterialCompilationFinished().RemoveAll(this);
+
+	UCEClonerSubsystem* ClonerSubsystem = UCEClonerSubsystem::Get();
+	if (!!AttachmentRoot && !!ClonerSubsystem)
+	{
+		if (const TSharedPtr<ICEClonerSceneTreeCustomResolver> CustomSceneTreeResolver = ClonerSubsystem->FindCustomLevelSceneTreeResolver(AttachmentRoot->GetLevel()))
+		{
+			CustomSceneTreeResolver->OnActorHierarchyChanged().RemoveAll(this);
+		}
+	}
+	
+	if (GEngine)
+	{
+		GEngine->OnLevelActorAttached().RemoveAll(this);
+		GEngine->OnLevelActorDetached().RemoveAll(this);
+	}
+#endif
+}
+
+void FCEClonerAttachmentGroupBehavior::GetOrderedChildrenActors(FCEClonerAttachmentTree& InTree, AActor* InActor, TArray<AActor*>& OutChildren) const
+{
+	UCEClonerSubsystem* ClonerSubsystem = UCEClonerSubsystem::Get();
+	if (!ClonerSubsystem)
+	{
+		return;
+	}
+
+	const TSharedPtr<ICEClonerSceneTreeCustomResolver> CustomSceneTreeResolver = ClonerSubsystem->FindCustomLevelSceneTreeResolver(InActor->GetLevel());
+
+	if (!CustomSceneTreeResolver.IsValid() || !CustomSceneTreeResolver->GetDirectChildrenActor(InActor, OutChildren))
+	{
+		InActor->GetAttachedActors(OutChildren, /** Reset */true, /** IncludeChildren */false);
+	}
+
+	OutChildren.RemoveAll([](const AActor* InActor)
+	{
+		return !IsValid(InActor);
+	});
+}
+
+void FCEClonerAttachmentGroupBehavior::InvalidateMesh(FCEClonerAttachmentTree& InTree, AActor* InActor)
+{
+#if WITH_EDITOR
+	if (!InActor)
+	{
+		return;
+	}
+
+	if (const FCEClonerAttachmentItem* FoundItem = InTree.ItemAttachmentMap.Find(InActor))
+	{
+		if (FoundItem->bRootItem || !FoundItem->ParentActor.IsValid())
+		{
+			const int32 RootIdx = InTree.RootActors.Find(InActor);
+			if (InTree.MergedBakedMeshes.IsValidIndex(RootIdx))
+			{
+				InTree.MergedBakedMeshes[RootIdx] = nullptr;
+				InTree.bItemAttachmentsDirty = true;
+			}
+		}
+		else
+		{
+			InvalidateMesh(InTree, FoundItem->ParentActor.Get());
+		}
+	}
+#endif
+}
+
+void FCEClonerAttachmentGroupBehavior::OnItemAttached(FCEClonerAttachmentTree& InTree, FCEClonerAttachmentItem& InItem)
+{
+	if (AActor* Actor = InItem.ItemActor.Get())
+	{
+		BindActorDelegates(InTree, Actor);
+
+#if WITH_EDITOR
+		// Ensure cached mesh bounds matches with attachment bounds
+		if (InItem.bRootItem)
+		{
+			AActor* ItemActor = InItem.ItemActor.Get();
+
+			const FBox BakedMeshBounds = InTree.GetCacheBounds(ItemActor);
+			const FBox AttachmentBounds = InTree.GetAttachmentBounds(ItemActor, /** IncludeChildren */true);
+
+			if (BakedMeshBounds.IsValid
+				&& AttachmentBounds.IsValid
+				&& !BakedMeshBounds.GetSize().Equals(AttachmentBounds.GetSize(), 1.f))
+			{
+				const AActor* ClonerActor = InTree.GetAttachmentRoot();
+
+				UE::ClonerEffector::Utilities::ShowWarning(FText::Format(
+					LOCTEXT("AttachmentCachedBoundsMismatch", "Cloner {0} : {1} cached bounds vs actor bounds mismatch, see logs"),
+					FText::FromString(ClonerActor->GetActorNameOrLabel()),
+					FText::FromString(Actor->GetActorNameOrLabel())));
+
+				UE_LOG(LogCECloner, Warning, TEXT("%s : %s bounds mismatch, cached bounds (%s) vs actor bounds (%s), verify attached actor"), *ClonerActor->GetActorNameOrLabel(), *Actor->GetActorNameOrLabel(), *BakedMeshBounds.ToString(), *AttachmentBounds.ToString());
+			}
+		}
+#endif
+	}
+}
+
+void FCEClonerAttachmentGroupBehavior::OnItemDetached(FCEClonerAttachmentTree& InTree, FCEClonerAttachmentItem& InItem)
+{
+	if (AActor* Actor = InItem.ItemActor.Get())
+	{
+		UnbindActorDelegates(InTree, Actor);
+	}
+}
+
+void FCEClonerAttachmentFlatBehavior::GetOrderedChildrenActors(FCEClonerAttachmentTree& InTree, AActor* InActor, TArray<AActor*>& OutChildren) const
+{
+	if (InActor != InTree.GetAttachmentRoot())
+	{
+		return;
+	}
+
+	UCEClonerSubsystem* ClonerSubsystem = UCEClonerSubsystem::Get();
+	if (!ClonerSubsystem)
+	{
+		return;
+	}
+
+	const TSharedPtr<ICEClonerSceneTreeCustomResolver> CustomSceneTreeResolver = ClonerSubsystem->FindCustomLevelSceneTreeResolver(InActor->GetLevel());
+
+	if (CustomSceneTreeResolver.IsValid())
+	{
+		TFunction<void(AActor*)> AddOrderedChildrenActorsRecursively = [&AddOrderedChildrenActorsRecursively, &CustomSceneTreeResolver, &OutChildren](AActor* InActor)
+		{
+			TArray<AActor*> Children;
+			CustomSceneTreeResolver->GetDirectChildrenActor(InActor, Children);
+
+			for (AActor* ChildActor : Children)
+			{
+				TArray<UPrimitiveComponent*> PrimitiveComponents;
+				ChildActor->GetComponents(PrimitiveComponents, /** IncludeChildrenActors */false);
+
+				for (UPrimitiveComponent* PrimitiveComponent : PrimitiveComponents)
+				{
+					// Skip invalid component or ones without any geometry
+					if (!!PrimitiveComponent && FCEMeshBuilder::HasAnyGeometry(PrimitiveComponent))
+					{
+						OutChildren.Add(ChildActor);
+						break;
+					}
+				}
+
+				AddOrderedChildrenActorsRecursively(ChildActor);
+			}
+		};
+
+		AddOrderedChildrenActorsRecursively(InActor);
+	}
+	else
+	{
+		InActor->GetAttachedActors(OutChildren, /** Reset */true, /** IncludeChildren */true);
+	}
+
+	OutChildren.RemoveAll([](const AActor* InActor)
+	{
+		return !IsValid(InActor);
+	});
+}
+
+void FCEClonerAttachmentGroupBehavior::BindActorDelegates(FCEClonerAttachmentTree& InTree, AActor* InActor)
+{
+	if (!IsValid(InActor))
+	{
+		return;
+	}
+
+#if WITH_EDITOR
+	// Detect static mesh change
+	TArray<UStaticMeshComponent*> StaticMeshComponents;
+	InActor->GetComponents(StaticMeshComponents, false);
+	for (UStaticMeshComponent* StaticMeshComponent : StaticMeshComponents)
+	{
+		if (!StaticMeshComponent->OnStaticMeshChanged().IsBoundToObject(this))
+		{
+			StaticMeshComponent->OnStaticMeshChanged().AddSPLambda(this, [this, &InTree, InActor](UStaticMeshComponent*)
+			{
+				OnMeshChanged(InTree, InActor);
+			});
+		}
+	}
+
+	// Detect dynamic mesh change
+	TArray<UDynamicMeshComponent*> DynamicMeshComponents;
+	InActor->GetComponents(DynamicMeshComponents, false);
+	for (UDynamicMeshComponent* DynamicMeshComponent : DynamicMeshComponents)
+	{
+		if (!DynamicMeshComponent->OnMeshChanged.IsBoundToObject(this))
+		{
+			DynamicMeshComponent->OnMeshChanged.AddSPLambda(this, [this, &InTree, InActor]()
+			{
+				OnMeshChanged(InTree, InActor);
+			});
+		}
+	}
+#endif
+
+	// Detect components transform
+	TArray<USceneComponent*> SceneComponents;
+	constexpr bool bIncludeChildren = false;
+	InActor->GetComponents(SceneComponents, bIncludeChildren);
+	for (USceneComponent* SceneComponent : SceneComponents)
+	{
+		if (!SceneComponent->TransformUpdated.IsBoundToObject(this))
+		{
+			SceneComponent->TransformUpdated.AddSPLambda(this, [this, &InTree](USceneComponent* InComponent, EUpdateTransformFlags InFlags, ETeleportType InTeleport)
+			{
+				OnComponentTransformed(InTree, InComponent, InFlags, InTeleport);
+			});
+		}
+	}
+}
+
+void FCEClonerAttachmentGroupBehavior::UnbindActorDelegates(FCEClonerAttachmentTree& InTree, AActor* InActor)
+{
+	if (!InActor)
+	{
+		return;
+	}
+
+#if WITH_EDITOR
+	TArray<UStaticMeshComponent*> StaticMeshComponents;
+	InActor->GetComponents(StaticMeshComponents, false);
+	for (UStaticMeshComponent* StaticMeshComponent : StaticMeshComponents)
+	{
+		StaticMeshComponent->OnStaticMeshChanged().RemoveAll(this);
+	}
+
+	TArray<UDynamicMeshComponent*> DynamicMeshComponents;
+	InActor->GetComponents(DynamicMeshComponents, false);
+	for (UDynamicMeshComponent* DynamicMeshComponent : DynamicMeshComponents)
+	{
+		DynamicMeshComponent->OnMeshChanged.RemoveAll(this);
+	}
+#endif
+
+	TArray<USceneComponent*> SceneComponents;
+	constexpr bool bIncludeChildren = false;
+	InActor->GetComponents(SceneComponents, bIncludeChildren);
+	for (USceneComponent* SceneComponent : SceneComponents)
+	{
+		SceneComponent->TransformUpdated.RemoveAll(this);
+	}
+}
+
+void FCEClonerAttachmentGroupBehavior::OnComponentTransformed(FCEClonerAttachmentTree& InTree, USceneComponent* InComponent, EUpdateTransformFlags InFlags, ETeleportType InTeleport)
+{
+	if (!InComponent || !InComponent->GetOwner() || InFlags == EUpdateTransformFlags::PropagateFromParent)
+	{
+		return;
+	}
+
+	AActor* Owner = InComponent->GetOwner();
+
+	FCEClonerAttachmentItem* Item = InTree.ItemAttachmentMap.Find(Owner);
+
+	if (!Item || !Item->bSetupDone)
+	{
+		return;
+	}
+
+	const AActor* RootActor = InTree.FindRootActor(Owner);
+
+	if (!RootActor)
+	{
+		return;
+	}
+
+	// Skip update if root component has moved, since we can simply offset the mesh
+	if (RootActor == Owner && RootActor->GetRootComponent() == InComponent)
+	{
+		InTree.bItemAttachmentsDirty = true;
+		return;
+	}
+
+	bool bComponentSupported = FCEMeshBuilder::IsComponentSupported(InComponent);
+
+	if (!bComponentSupported)
+	{
+		for (const TObjectPtr<USceneComponent>& ChildComponent : InComponent->GetAttachChildren())
+		{
+			if (FCEMeshBuilder::IsComponentSupported(ChildComponent))
+			{
+				bComponentSupported = true;
+				break;
+			}
+		}
+	}
+
+	if (!bComponentSupported)
+	{
+		return;
+	}
+
+	const AActor* ClonerActor = InTree.GetAttachmentRoot();
+
+	if (!ClonerActor)
+	{
+		return;
+	}
+
+	UE_LOG(LogCECloner, Log, TEXT("%s : Transform state changed for %s"), *ClonerActor->GetActorNameOrLabel(), *Owner->GetActorNameOrLabel());
+
+	Item->MeshStatus = ECEClonerAttachmentStatus::Outdated;
+	InvalidateMesh(InTree, Owner);
+	InTree.DirtyItemAttachments.Add(Item->ItemActor);
+	InTree.MarkAttachmentOutdated();
+}
+
+void FCEClonerAttachmentGroupBehavior::OnWorldActorDestroyed(FCEClonerAttachmentTree& InTree, AActor* InActor)
+{
+	if (InActor && InTree.ItemAttachmentMap.Contains(InActor))
+	{
+		InTree.DetachItem(InActor);
+	}
+}
+
+#if WITH_EDITOR
+void FCEClonerAttachmentGroupBehavior::OnMeshChanged(FCEClonerAttachmentTree& InTree, AActor* InActor)
+{
+	if (!InActor)
+	{
+		return;
+	}
+
+	const AActor* ClonerActor = InTree.GetAttachmentRoot();
+
+	if (!ClonerActor)
+	{
+		return;
+	}
+
+	FCEClonerAttachmentItem* Item = InTree.ItemAttachmentMap.Find(InActor);
+
+	if (!Item || !Item->bSetupDone || Item->MeshStatus == ECEClonerAttachmentStatus::Outdated)
+	{
+		return;
+	}
+
+	UE_LOG(LogCECloner, Log, TEXT("%s : Detected mesh change for %s"), *ClonerActor->GetActorNameOrLabel(), *InActor->GetActorNameOrLabel());
+
+	Item->MeshStatus = ECEClonerAttachmentStatus::Outdated;
+	InvalidateMesh(InTree, InActor);
+	InTree.DirtyItemAttachments.Add(Item->ItemActor);
+	InTree.MarkAttachmentOutdated();
+}
+
+void FCEClonerAttachmentGroupBehavior::OnRenderStateDirty(FCEClonerAttachmentTree& InTree, UActorComponent& InComponent)
+{
+	AActor* Owner = InComponent.GetOwner();
+	const AActor* ClonerActor = InTree.GetAttachmentRoot();
+
+	if (!Owner || Owner->GetLevel() != ClonerActor->GetLevel())
+	{
+		return;
+	}
+
+	// Does it contain geometry that we can convert
+	if (!FCEMeshBuilder::IsComponentSupported(&InComponent))
+	{
+		return;
+	}
+
+	FCEClonerAttachmentItem* Item = InTree.ItemAttachmentMap.Find(Owner);
+
+	if (!Item || !Item->bSetupDone || !Item->CheckBoundsChanged(/** CheckAndUpdate */false))
+	{
+		return;
+	}
+
+	UE_LOG(LogCECloner, Log, TEXT("%s : Render state changed for %s"), *ClonerActor->GetActorNameOrLabel(), *Owner->GetActorNameOrLabel());
+
+	// Rebind delegates as new components might be available
+	BindActorDelegates(InTree, Owner);
+
+	Item->MeshStatus = ECEClonerAttachmentStatus::Outdated;
+	InvalidateMesh(InTree, Owner);
+	InTree.DirtyItemAttachments.Add(Item->ItemActor);
+	InTree.MarkAttachmentOutdated();
+}
+
+void FCEClonerAttachmentGroupBehavior::OnCheckMaterialChanged(FCEClonerAttachmentTree& InTree, UObject* InObject)
+{
+	if (!IsValid(InObject))
+	{
+		return;
+	}
+
+	const AActor* ClonerActor = InTree.GetAttachmentRoot();
+
+	if (!ClonerActor)
+	{
+		return;
+	}
+
+	AActor* ActorChanged = Cast<AActor>(InObject);
+	ActorChanged = ActorChanged ? ActorChanged : InObject->GetTypedOuter<AActor>();
+
+	if (!ActorChanged)
+	{
+		return;
+	}
+
+	FCEClonerAttachmentItem* AttachmentItem = InTree.ItemAttachmentMap.Find(ActorChanged);
+
+	if (!AttachmentItem || !AttachmentItem->bSetupDone)
+	{
+		return;
+	}
+
+	TArray<TWeakObjectPtr<UMaterialInterface>> UnsupportedMaterials;
+	if (AttachmentItem->CheckMaterialsChanged(/** Update */true, &UnsupportedMaterials))
+	{
+		// Show warning for unset materials
+		if (!UnsupportedMaterials.IsEmpty())
+		{
+			if (UCEClonerSubsystem* ClonerSubsystem = UCEClonerSubsystem::Get())
+			{
+				ClonerSubsystem->FireMaterialWarning(ClonerActor, ActorChanged, UnsupportedMaterials);
+			}
+		}
+
+		UE_LOG(LogCECloner, Log, TEXT("%s : Detected material change for %s"), *ClonerActor->GetActorNameOrLabel(), *ActorChanged->GetActorNameOrLabel());
+		InvalidateMesh(InTree, ActorChanged);
+		InTree.MarkAttachmentOutdated();
+	}
+}
+
+void FCEClonerAttachmentGroupBehavior::OnLevelHierarchyChanged(FCEClonerAttachmentTree& InTree, AActor* InActor)
+{
+	if (IsValid(InActor))
+	{
+		const AActor* AttachmentRoot = InTree.GetAttachmentRoot();
+
+		if (InTree.ItemAttachmentMap.Contains(InActor) || InActor->IsAttachedTo(AttachmentRoot))
+		{
+			InTree.MarkAttachmentOutdated();
+		}
+	}
+}
+#endif
+
+#undef LOCTEXT_NAMESPACE
diff --git a/Engine/Plugins/Experimental/ClonerEffector/Source/ClonerEffector/Private/Cloner/CEClonerActor.cpp b/Engine/Plugins/Experimental/ClonerEffector/Source/ClonerEffector/Private/Cloner/CEClonerActor.cpp
--- a/Engine/Plugins/Experimental/ClonerEffector/Source/ClonerEffector/Private/Cloner/CEClonerActor.cpp
+++ b/Engine/Plugins/Experimental/ClonerEffector/Source/ClonerEffector/Private/Cloner/CEClonerActor.cpp
@@ -22,6 +22,7 @@
 #include "Cloner/Layouts/CEClonerSphereRandomLayout.h"
 #include "Cloner/Layouts/CEClonerSphereUniformLayout.h"
 #include "Cloner/Layouts/CEClonerSplineLayout.h"
+#include "Cloner/Logs/CEClonerLogs.h"
 #include "Containers/Ticker.h"
 #include "Effector/CEEffectorActor.h"
 #include "Materials/MaterialInterface.h"
@@ -33,8 +34,6 @@
 #include "Settings/LevelEditorViewportSettings.h"
 #endif
 
-DEFINE_LOG_CATEGORY_STATIC(LogCEClonerActor, Log, All);
-
 struct FCEClonerActorVersion
 {
 	enum Type : int32
@@ -119,14 +118,13 @@ void ACEClonerActor::MigrateDeprecatedProperties()
 		return;
 	}
 
-	UE_LOG(LogCEClonerActor, Warning, TEXT("%s : Cloner migrating from version %i to latest %i, please re-save this asset"), *GetActorNameOrLabel(), MigrateToVersion, FCEClonerActorVersion::LatestVersion);
+	UE_LOG(LogCECloner, Warning, TEXT("%s : Cloner migrating from version %i to latest %i, please re-save this asset"), *GetActorNameOrLabel(), MigrateToVersion, FCEClonerActorVersion::LatestVersion);
 
 	PRAGMA_DISABLE_DEPRECATION_WARNINGS
 
 	if (MigrateToVersion <= FCEClonerActorVersion::ComponentMigration)
 	{
 		ClonerComponent->SetEnabled(bEnabled);
-		ClonerComponent->SetTreeUpdateInterval(TreeUpdateInterval);
 		ClonerComponent->SetSeed(Seed);
 		ClonerComponent->SetColor(Color);
 		ClonerComponent->SetLayoutName(LayoutName);
diff --git a/Engine/Plugins/Experimental/ClonerEffector/Source/ClonerEffector/Private/Cloner/CEClonerComponent.cpp b/Engine/Plugins/Experimental/ClonerEffector/Source/ClonerEffector/Private/Cloner/CEClonerComponent.cpp
--- a/Engine/Plugins/Experimental/ClonerEffector/Source/ClonerEffector/Private/Cloner/CEClonerComponent.cpp
+++ b/Engine/Plugins/Experimental/ClonerEffector/Source/ClonerEffector/Private/Cloner/CEClonerComponent.cpp
@@ -3,6 +3,7 @@
 #include "Cloner/CEClonerComponent.h"
 
 #include "Async/Async.h"
+#include "Cloner/Attachments/CEClonerAttachmentTreeBehavior.h"
 #include "Cloner/CEClonerActor.h"
 #include "Cloner/Extensions/CEClonerExtensionBase.h"
 #include "Cloner/Layouts/CEClonerLayoutBase.h"
@@ -13,6 +14,7 @@
 #include "Engine/StaticMeshActor.h"
 #include "Engine/Texture2D.h"
 #include "Engine/World.h"
+#include "Logs/CEClonerLogs.h"
 #include "Materials/Material.h"
 #include "Materials/MaterialInstanceDynamic.h"
 #include "NiagaraMeshRendererProperties.h"
@@ -21,27 +23,24 @@
 #include "Subsystems/CEClonerSubsystem.h"
 #include "UDynamicMesh.h"
 #include "UObject/Package.h"
+#include "Utilities/CEClonerEffectorUtilities.h"
 
 #if WITH_EDITOR
 #include "Editor/EditorEngine.h"
-#include "Framework/Notifications/NotificationManager.h"
 #include "Materials/Material.h"
 #include "Misc/MessageDialog.h"
 #include "Misc/ScopedSlowTask.h"
 #include "ScopedTransaction.h"
-#include "Widgets/Notifications/SNotificationList.h"
 #endif
 
+#define LOCTEXT_NAMESPACE "CEClonerComponent"
+
 UCEClonerComponent::FOnClonerInitialized UCEClonerComponent::OnClonerInitializedDelegate;
 UCEClonerComponent::FOnClonerLayoutLoaded UCEClonerComponent::OnClonerLayoutLoadedDelegate;
 UCEClonerComponent::FOnClonerMeshUpdated UCEClonerComponent::OnClonerMeshUpdatedDelegate;
 UCEClonerComponent::FOnClonerAttachmentChanged UCEClonerComponent::OnClonerActorAttachedDelegate;
 UCEClonerComponent::FOnClonerAttachmentChanged UCEClonerComponent::OnClonerActorDetachedDelegate;
 
-DEFINE_LOG_CATEGORY_STATIC(LogCEClonerComponent, Log, All);
-
-#define LOCTEXT_NAMESPACE "CEClonerComponent"
-
 UCEClonerComponent::UCEClonerComponent()
 	: UNiagaraComponent()
 {
@@ -69,18 +68,17 @@ UCEClonerComponent::UCEClonerComponent()
 	{
 		UCEClonerSubsystem::OnClonerSetEnabled().AddUObject(this, &UCEClonerComponent::OnClonerSetEnabled);
 
-		// Bind to delegate to detect material changes
-#if WITH_EDITOR
-		USceneComponent::MarkRenderStateDirtyEvent.AddUObject(this, &UCEClonerComponent::OnRenderStateDirty);
-
-		FCoreUObjectDelegates::OnObjectPropertyChanged.AddUObject(this, &UCEClonerComponent::OnActorPropertyChanged);
-
-		UMaterial::OnMaterialCompilationFinished().AddUObject(this, &UCEClonerComponent::OnMaterialCompiled);
-#endif
-
 		// Apply default layout
 		const TArray<FName> LayoutNames = GetClonerLayoutNames();
 		LayoutName = !LayoutNames.IsEmpty() ? LayoutNames[0] : NAME_None;
+
+		// Apply default behavior
+		const TArray<FName> BehaviorNames = GetClonerTreeBehaviorNames();
+		TreeBehaviorName = !BehaviorNames.IsEmpty() ? BehaviorNames[0] : NAME_None;
+
+		// Bind attachment tree events
+		ClonerTree.OnItemAttached().BindUObject(this, &UCEClonerComponent::OnTreeItemAttached);
+		ClonerTree.OnItemDetached().BindUObject(this, &UCEClonerComponent::OnTreeItemDetached);
 	}
 }
 
@@ -121,22 +119,29 @@ void UCEClonerComponent::PostLoad()
 }
 
 #if WITH_EDITOR
-void UCEClonerComponent::PostEditImport()
+void UCEClonerComponent::PreDuplicate(FObjectDuplicationParameters& InParams)
+{
+	Super::PreDuplicate(InParams);
+
+	ForceUpdateCloner();
+}
+
+void UCEClonerComponent::PostDuplicate(bool bInPIE)
 {
 	SetAsset(nullptr);
 
-	Super::PostEditImport();
+	Super::PostDuplicate(bInPIE);
 
 	RegisterTicker();
 
 	ForceUpdateCloner();
 }
 
-void UCEClonerComponent::PostDuplicate(bool bInPIE)
+void UCEClonerComponent::PostEditImport()
 {
 	SetAsset(nullptr);
 
-	Super::PostDuplicate(bInPIE);
+	Super::PostEditImport();
 
 	RegisterTicker();
 
@@ -161,6 +166,7 @@ const TCEPropertyChangeDispatcher<UCEClonerComponent> UCEClonerComponent::Proper
 	{ GET_MEMBER_NAME_CHECKED(UCEClonerComponent, Color), &UCEClonerComponent::OnColorChanged },
 	{ GET_MEMBER_NAME_CHECKED(UCEClonerComponent, GlobalScale), &UCEClonerComponent::OnGlobalScaleChanged },
 	{ GET_MEMBER_NAME_CHECKED(UCEClonerComponent, GlobalRotation), &UCEClonerComponent::OnGlobalRotationChanged },
+	{ GET_MEMBER_NAME_CHECKED(UCEClonerComponent, TreeBehaviorName), &UCEClonerComponent::OnTreeBehaviorNameChanged },
 	{ GET_MEMBER_NAME_CHECKED(UCEClonerComponent, bVisualizerSpriteVisible), &UCEClonerComponent::OnVisualizerSpriteVisibleChanged },
 	/** Layout */
 	{ GET_MEMBER_NAME_CHECKED(UCEClonerComponent, LayoutName), &UCEClonerComponent::OnLayoutNameChanged },
@@ -185,13 +191,7 @@ void UCEClonerComponent::OnComponentDestroyed(bool bInDestroyingHierarchy)
 {
 	Super::OnComponentDestroyed(bInDestroyingHierarchy);
 
-#if WITH_EDITOR
-	USceneComponent::MarkRenderStateDirtyEvent.RemoveAll(this);
-	
-	FCoreUObjectDelegates::OnObjectPropertyChanged.RemoveAll(this);
-	
-	UMaterial::OnMaterialCompilationFinished().RemoveAll(this);
-#endif
+	ClonerTree.Cleanup();
 
 	FWorldDelegates::LevelAddedToWorld.RemoveAll(this);
 
@@ -210,7 +210,7 @@ void UCEClonerComponent::UpdateClonerRenderState()
 	if (!GetAsset()
 		|| IsGarbageCollectingAndLockingUObjectHashTables()
 		|| bClonerMeshesUpdating
-		|| !bClonerMeshesDirty
+		|| !ClonerTree.bItemAttachmentsDirty
 		|| ClonerTree.Status != ECEClonerAttachmentStatus::Updated)
 	{
 		return;
@@ -221,545 +221,22 @@ void UCEClonerComponent::UpdateClonerRenderState()
 #else
 	OnDirtyMeshesUpdated(true);
 #endif
-}
-
-void UCEClonerComponent::UpdateClonerAttachmentTree(bool bInReset)
-{
-	if (ClonerTree.Status == ECEClonerAttachmentStatus::Updated)
-	{
-		ClonerTree.Status = ECEClonerAttachmentStatus::Outdated;
-	}
-
-	if (bInReset)
-	{
-		ClonerTree.Reset();
-		ClonerTree.Status = ECEClonerAttachmentStatus::Outdated;
-	}
-
-	UpdateAttachmentTree();
-}
-
-void UCEClonerComponent::UpdateAttachmentTree()
-{
-	if (ClonerTree.Status != ECEClonerAttachmentStatus::Outdated)
-	{
-		return;
-	}
-	ClonerTree.Status = ECEClonerAttachmentStatus::Updating;
-
-	// Invalidate all, to see what is outdated and what is still invalid
-	for (TPair<TWeakObjectPtr<AActor>, FCEClonerAttachmentItem>& AttachmentPair : ClonerTree.ItemAttachmentMap)
-	{
-		AttachmentPair.Value.Status = ECEClonerAttachmentStatus::Invalid;
-	}
-
-	// Update root attachment items
-	TArray<AActor*> RootChildren;
-	GetOrderedRootActors(RootChildren);
-
-	TArray<TObjectPtr<UStaticMesh>> NewCombinesMeshes;
-	TArray<TWeakObjectPtr<AActor>> NewRootActors;
-	for (int32 RootIdx = 0; RootIdx < RootChildren.Num(); RootIdx++)
-	{
-		AActor* RootChild = RootChildren[RootIdx];
-
-		if (!IsValid(RootChild))
-		{
-			continue;
-		}
-
-		UpdateActorAttachment(RootChild, nullptr);
-
-		// Lets find the old root idx
-		const int32 OldIdx = ClonerTree.RootActors.Find(RootChild);
-		TObjectPtr<UStaticMesh> CombineBakedMesh = nullptr;
-		if (ClonerTree.MergedBakedMeshes.IsValidIndex(OldIdx))
-		{
-			CombineBakedMesh = ClonerTree.MergedBakedMeshes[OldIdx].Get();
-
-			// Did we rearrange stuff ?
-			if (RootIdx != OldIdx)
-			{
-				bClonerMeshesDirty = true;
-			}
-		}
-		NewCombinesMeshes.Add(CombineBakedMesh);
-		NewRootActors.Add(RootChild);
-	}
-
-	// Did we remove any root actors ?
-	if (ClonerTree.RootActors.Num() != NewRootActors.Num())
-	{
-		bClonerMeshesDirty = true;
-	}
-
-	// Did we need to update meshes
-	TArray<TWeakObjectPtr<AActor>> ClonedActors;
-	ClonerTree.ItemAttachmentMap.GenerateKeyArray(ClonedActors);
-	for (const TWeakObjectPtr<AActor>& ClonedActorWeak : ClonedActors)
-	{
-		FCEClonerAttachmentItem* ClonedItem = ClonerTree.ItemAttachmentMap.Find(ClonedActorWeak);
-
-		if (!ClonedItem)
-		{
-			continue;
-		}
-
-		AActor* ClonedActor = ClonedActorWeak.Get();
-
-		if (ClonedItem->Status == ECEClonerAttachmentStatus::Invalid)
-		{
-			InvalidateBakedStaticMesh(ClonedActor);
-			UnbindActorDelegates(ClonedActor);
-			ClonerTree.ItemAttachmentMap.Remove(ClonedItem->ItemActor);
-			SetActorVisibility(ClonedActor, true);
-
-			OnClonerActorDetachedDelegate.Broadcast(this, ClonedActor);
-		}
-		else if (ClonedItem->Status == ECEClonerAttachmentStatus::Outdated)
-		{
-			if (ClonedItem->MeshStatus == ECEClonerAttachmentStatus::Outdated)
-			{
-				ClonerTree.DirtyItemAttachments.Add(ClonedItem->ItemActor);
-				InvalidateBakedStaticMesh(ClonedActor);
-			}
-			bClonerMeshesDirty = true;
-			ClonedItem->Status = ECEClonerAttachmentStatus::Updated;
-		}
-	}
-
-	// Did we remove an attachment ?
-	if (ClonedActors.Num() != ClonerTree.ItemAttachmentMap.Num())
-	{
-		bClonerMeshesDirty = true;
-	}
-
-	if (!ClonerTree.DirtyItemAttachments.IsEmpty())
-	{
-		bClonerMeshesDirty = true;
-	}
-
-	ClonerTree.RootActors = NewRootActors;
-	ClonerTree.MergedBakedMeshes = NewCombinesMeshes;
-	ClonerTree.Status = ECEClonerAttachmentStatus::Updated;
-}
-
-void UCEClonerComponent::UpdateActorAttachment(AActor* InActor, AActor* InParent)
-{
-	if (!IsValid(InActor))
-	{
-		return;
-	}
-
-	const AActor* ClonerActor = GetOwner();
-	const FTransform& ClonerTransform = ClonerActor->GetActorTransform();
-
-	// Here order is not important
-	TArray<AActor*> ChildrenActors;
-	InActor->GetAttachedActors(ChildrenActors, true, false);
-
-	FCEClonerAttachmentItem* AttachmentItem = ClonerTree.ItemAttachmentMap.Find(InActor);
-	if (AttachmentItem)
-	{
-		AttachmentItem->Status = ECEClonerAttachmentStatus::Updated;
-
-		// Check Root is the same
-		const bool bIsRoot = InParent == nullptr;
-		if (AttachmentItem->bRootItem != bIsRoot)
-		{
-			InvalidateBakedStaticMesh(InActor);
-			AttachmentItem->bRootItem = bIsRoot;
-			AttachmentItem->Status = ECEClonerAttachmentStatus::Outdated;
-		}
-
-		// Check parent is the same
-		if (AttachmentItem->ParentActor.Get() != InParent)
-		{
-			InvalidateBakedStaticMesh(InParent);
-			InvalidateBakedStaticMesh(AttachmentItem->ParentActor.Get());
-			AttachmentItem->ParentActor = InParent;
-			AttachmentItem->Status = ECEClonerAttachmentStatus::Outdated;
-		}
-
-		// Check transform is the same
-		const FTransform ActorTransform = InActor->GetActorTransform().GetRelativeTransform(ClonerTransform);
-		if (!ActorTransform.Equals(AttachmentItem->ActorTransform))
-		{
-			// invalidate if not root, else change transform in mesh renderer
-			if (!bIsRoot)
-			{
-				InvalidateBakedStaticMesh(InActor);
-			}
-			AttachmentItem->ActorTransform = ActorTransform;
-			AttachmentItem->Status = ECEClonerAttachmentStatus::Outdated;
-		}
-
-#if WITH_EDITOR
-		// Check bounds change
-		if (AttachmentItem->CheckBoundsChanged(/** CheckAndUpdate */true))
-		{
-			InvalidateBakedStaticMesh(InActor);
-			AttachmentItem->Status = ECEClonerAttachmentStatus::Outdated;
-		}
-#endif
-	}
-	else
-	{
-		AttachmentItem = &ClonerTree.ItemAttachmentMap.Add(InActor, FCEClonerAttachmentItem());
-		AttachmentItem->ItemActor = InActor;
-		AttachmentItem->ParentActor = InParent;
-		AttachmentItem->ActorTransform = InActor->GetActorTransform().GetRelativeTransform(ClonerTransform);
-		AttachmentItem->MeshStatus = ECEClonerAttachmentStatus::Outdated;
-		AttachmentItem->bRootItem = InParent == nullptr;
-		AttachmentItem->Status = ECEClonerAttachmentStatus::Outdated;
-
-#if WITH_EDITOR
-		AttachmentItem->CheckBoundsChanged(/** CheckAndUpdate */true);
-#endif
-
-		InvalidateBakedStaticMesh(InActor);
-	}
-
-	if (!AttachmentItem->bSetupDone)
-	{
-		AttachmentItem->bSetupDone = true;
-
-		BindActorDelegates(InActor);
-		SetActorVisibility(InActor, false);
-
-		OnClonerActorAttachedDelegate.Broadcast(this, InActor);
-	}
-
-	if (AttachmentItem->ChildrenActors.Num() != ChildrenActors.Num())
-	{
-		InvalidateBakedStaticMesh(InActor);
-	}
-
-	AttachmentItem->ChildrenActors.Empty(ChildrenActors.Num());
-	for (AActor* ChildActor : ChildrenActors)
-	{
-		AttachmentItem->ChildrenActors.Add(ChildActor);
-		UpdateActorAttachment(ChildActor, InActor);
-	}
-}
-
-void UCEClonerComponent::BindActorDelegates(AActor* InActor)
-{
-	if (!IsValid(InActor))
-	{
-		return;
-	}
-
-	InActor->OnDestroyed.AddUniqueDynamic(this, &UCEClonerComponent::OnActorDestroyed);
-
-#if WITH_EDITOR
-	// Detect static mesh change
-	TArray<UStaticMeshComponent*> StaticMeshComponents;
-	InActor->GetComponents(StaticMeshComponents, false);
-	for (UStaticMeshComponent* StaticMeshComponent : StaticMeshComponents)
-	{
-		if (!StaticMeshComponent->OnStaticMeshChanged().IsBoundToObject(this))
-		{
-			StaticMeshComponent->OnStaticMeshChanged().AddUObject(this, &UCEClonerComponent::OnMeshChanged, InActor);
-		}
-	}
-
-	// Detect dynamic mesh change
-	TArray<UDynamicMeshComponent*> DynamicMeshComponents;
-	InActor->GetComponents(DynamicMeshComponents, false);
-	for (UDynamicMeshComponent* DynamicMeshComponent : DynamicMeshComponents)
-	{
-		if (!DynamicMeshComponent->OnMeshChanged.IsBoundToObject(this))
-		{
-			UStaticMeshComponent* NullComponent = nullptr;
-			DynamicMeshComponent->OnMeshChanged.AddUObject(this, &UCEClonerComponent::OnMeshChanged, NullComponent, InActor);
-		}
-	}
-#endif
-
-	// Detect components transform
-	TArray<USceneComponent*> SceneComponents;
-	InActor->GetComponents(SceneComponents, /** IncludeChildren */false);
-	for (USceneComponent* SceneComponent : SceneComponents)
-	{
-		if (!SceneComponent->TransformUpdated.IsBoundToObject(this))
-		{
-			SceneComponent->TransformUpdated.AddUObject(this, &UCEClonerComponent::OnComponentTransformed);
-		}
-	}
-}
-
-void UCEClonerComponent::UnbindActorDelegates(AActor* InActor) const
-{
-	if (!InActor)
-	{
-		return;
-	}
-
-	InActor->OnDestroyed.RemoveAll(this);
-
-#if WITH_EDITOR
-	TArray<UStaticMeshComponent*> StaticMeshComponents;
-	InActor->GetComponents(StaticMeshComponents, false);
-	for (UStaticMeshComponent* StaticMeshComponent : StaticMeshComponents)
-	{
-		StaticMeshComponent->OnStaticMeshChanged().RemoveAll(this);
-	}
-
-	TArray<UDynamicMeshComponent*> DynamicMeshComponents;
-	InActor->GetComponents(DynamicMeshComponents, false);
-	for (UDynamicMeshComponent* DynamicMeshComponent : DynamicMeshComponents)
-	{
-		DynamicMeshComponent->OnMeshChanged.RemoveAll(this);
-	}
-#endif
-
-	TArray<USceneComponent*> SceneComponents;
-	InActor->GetComponents(SceneComponents, /** IncludeChildren */false);
-	for (USceneComponent* SceneComponent : SceneComponents)
-	{
-		SceneComponent->TransformUpdated.RemoveAll(this);
-	}
-}
-
-void UCEClonerComponent::SetActorVisibility(AActor* InActor, bool bInVisibility, ECEClonerActorVisibility InTarget)
-{
-	if (!InActor)
-	{
-		return;
-	}
-
-#if WITH_EDITOR
-	if (EnumHasAnyFlags(InTarget, ECEClonerActorVisibility::Editor))
-	{
-		InActor->SetIsTemporarilyHiddenInEditor(!bInVisibility);
-	}
-#endif
-
-	if (EnumHasAnyFlags(InTarget, ECEClonerActorVisibility::Game))
-	{
-		InActor->SetActorHiddenInGame(!bInVisibility);
-	}
-}
-
-void UCEClonerComponent::OnActorDestroyed(AActor* InDestroyedActor)
-{
-	if (ClonerTree.ItemAttachmentMap.Contains(InDestroyedActor))
-	{
-		InvalidateBakedStaticMesh(InDestroyedActor);
-		UnbindActorDelegates(InDestroyedActor);
-		ClonerTree.ItemAttachmentMap.Remove(InDestroyedActor);
-		SetActorVisibility(InDestroyedActor, true);
-		bClonerMeshesDirty = true;
-
-		OnClonerActorDetachedDelegate.Broadcast(this, InDestroyedActor);
-	}
-}
-
-#if WITH_EDITOR
-void UCEClonerComponent::OnActorPropertyChanged(UObject* InObject, FPropertyChangedEvent& InPropertyChangedEvent)
-{
-	OnMaterialChanged(InObject);
-}
-
-void UCEClonerComponent::OnMaterialCompiled(UMaterialInterface* InMaterial)
-{
-	OnMaterialChanged(InMaterial);
-}
-#endif
-
-void UCEClonerComponent::OnMaterialChanged(UObject* InObject)
-{
-	if (!IsValid(InObject))
-	{
-		return;
-	}
-
-	const AActor* ClonerActor = GetOwner();
-
-	if (!ClonerActor)
-	{
-		return;
-	}
-
-	AActor* ActorChanged = Cast<AActor>(InObject);
-	ActorChanged = ActorChanged ? ActorChanged : InObject->GetTypedOuter<AActor>();
-
-	if (!ActorChanged)
-	{
-		return;
-	}
-
-	FCEClonerAttachmentItem* AttachmentItem = ClonerTree.ItemAttachmentMap.Find(ActorChanged);
-
-	if (!AttachmentItem)
-	{
-		return;
-	}
-
-	TArray<UPrimitiveComponent*> PrimitiveComponents;
-	ActorChanged->GetComponents(PrimitiveComponents, /** IncludeChildrenActors */false);
-
-	int32 MatIdx = 0;
-	bool bMaterialChanged = false;
-	TArray<TWeakObjectPtr<UMaterialInterface>> NewMaterials;
-	NewMaterials.Reserve(PrimitiveComponents.Num());
-	UMaterialInterface* DefaultMaterial = LoadObject<UMaterialInterface>(nullptr, UCEClonerEffectorSettings::DefaultMaterialPath);
-
-	TArray<TWeakObjectPtr<UMaterialInterface>> UnsetMaterialsWeak;
-	for (UPrimitiveComponent* PrimitiveComponent : PrimitiveComponents)
-	{
-		if (!PrimitiveComponent || !FCEMeshBuilder::HasAnyGeometry(PrimitiveComponent))
-		{
-			continue;
-		}
-
-		for (int32 MatIndex = 0; MatIndex < PrimitiveComponent->GetNumMaterials(); MatIndex++)
-		{
-			UMaterialInterface* PreviousMaterial = PrimitiveComponent->GetMaterial(MatIndex);
-			UMaterialInterface* NewMaterial = PreviousMaterial;
-
-			if (FilterSupportedMaterial(NewMaterial, DefaultMaterial))
-			{
-				UnsetMaterialsWeak.Add(PreviousMaterial);
-			}
-
-			if (!AttachmentItem->BakedMaterials.IsValidIndex(MatIdx)
-				|| AttachmentItem->BakedMaterials[MatIdx] != NewMaterial)
-			{
-				bMaterialChanged = true;
-			}
-
-			NewMaterials.Add(NewMaterial);
-			MatIdx++;
-		}
-	}
-
-	// Show warning for unset materials
-	if (!UnsetMaterialsWeak.IsEmpty())
-	{
-		FireMaterialWarning(ActorChanged, UnsetMaterialsWeak);
-	}
-
-	if (bMaterialChanged)
-	{
-		UE_LOG(LogCEClonerComponent, Log, TEXT("%s : Detected material change for %s"), *ClonerActor->GetActorNameOrLabel(), *ActorChanged->GetActorNameOrLabel());
-
-		if (NewMaterials.Num() == AttachmentItem->BakedMaterials.Num())
-		{
-			AttachmentItem->BakedMaterials = NewMaterials;
-		}
-		else
-		{
-			AttachmentItem->MeshStatus = ECEClonerAttachmentStatus::Outdated;
-		}
-
-		InvalidateBakedStaticMesh(ActorChanged);
-	}
-}
-
-void UCEClonerComponent::OnMeshChanged(UStaticMeshComponent*, AActor* InActor)
-{
-	if (!InActor)
-	{
-		return;
-	}
-
-	const AActor* ClonerActor = GetOwner();
-
-	if (!ClonerActor)
-	{
-		return;
-	}
-
-	FCEClonerAttachmentItem* Item = ClonerTree.ItemAttachmentMap.Find(InActor);
-
-	if (!Item || Item->MeshStatus == ECEClonerAttachmentStatus::Outdated)
-	{
-		return;
-	}
-
-	UE_LOG(LogCEClonerComponent, Log, TEXT("%s : Detected mesh change for %s"), *ClonerActor->GetActorNameOrLabel(), *InActor->GetActorNameOrLabel());
-
-	Item->MeshStatus = ECEClonerAttachmentStatus::Outdated;
-	InvalidateBakedStaticMesh(InActor);
-	ClonerTree.DirtyItemAttachments.Add(Item->ItemActor);
-}
-
-void UCEClonerComponent::GetOrderedRootActors(TArray<AActor*>& OutActors) const
-{
-	const AActor* ClonerActor = GetOwner();
-	if (!ClonerActor)
-	{
-		return;
-	}
-
-	UCEClonerSubsystem* ClonerSubsystem = UCEClonerSubsystem::Get();
-	if (!ClonerSubsystem)
-	{
-		return;
-	}
-
-	const UCEClonerSubsystem::FOnGetOrderedActors& CustomActorResolver = ClonerSubsystem->GetCustomActorResolver();
-
-	if (CustomActorResolver.IsBound())
-	{
-		OutActors = CustomActorResolver.Execute(ClonerActor);
-	}
-	else
-	{
-		ClonerActor->GetAttachedActors(OutActors, true, false);
-	}
-
-	OutActors.RemoveAll([](const AActor* InActor)
-	{
-		return !IsValid(InActor);
-	});
-}
-
-AActor* UCEClonerComponent::GetRootActor(AActor* InActor) const
-{
-	if (InActor == nullptr)
-	{
-		return nullptr;
-	}
-	if (const FCEClonerAttachmentItem* Item = ClonerTree.ItemAttachmentMap.Find(InActor))
-	{
-		if (Item->bRootItem)
-		{
-			return InActor;
-		}
+}
 
-		return GetRootActor(Item->ParentActor.Get());
-	}
-	return nullptr;
+#if WITH_EDITOR
+FName UCEClonerComponent::GetActiveExtensionsPropertyName()
+{
+	return GET_MEMBER_NAME_CHECKED(UCEClonerComponent, ActiveExtensions);
 }
 
-void UCEClonerComponent::InvalidateBakedStaticMesh(AActor* InActor)
+FName UCEClonerComponent::GetActiveLayoutPropertyName()
 {
-#if WITH_EDITOR
-	if (!InActor)
-	{
-		return;
-	}
+	return GET_MEMBER_NAME_CHECKED(UCEClonerComponent, ActiveLayout);
+}
 
-	if (const FCEClonerAttachmentItem* FoundItem = ClonerTree.ItemAttachmentMap.Find(InActor))
-	{
-		if (FoundItem->bRootItem || !FoundItem->ParentActor.IsValid())
-		{
-			const int32 RootIdx = ClonerTree.RootActors.Find(InActor);
-			if (ClonerTree.MergedBakedMeshes.IsValidIndex(RootIdx))
-			{
-				ClonerTree.MergedBakedMeshes[RootIdx] = nullptr;
-				bClonerMeshesDirty = true;
-			}
-		}
-		else
-		{
-			InvalidateBakedStaticMesh(FoundItem->ParentActor.Get());
-		}
-	}
-#endif
+FName UCEClonerComponent::GetLayoutNamePropertyName()
+{
+	return GET_MEMBER_NAME_CHECKED(UCEClonerComponent, LayoutName);
 }
 
 void UCEClonerComponent::UpdateDirtyMeshesAsync()
@@ -773,6 +250,8 @@ void UCEClonerComponent::UpdateDirtyMeshesAsync()
 
 	TSet<TWeakObjectPtr<AActor>> DirtyAttachments = ClonerTree.DirtyItemAttachments;
 	ClonerTree.DirtyItemAttachments.Empty(DirtyAttachments.Num());
+	
+	UE_LOG(LogCECloner, Verbose, TEXT("%s : Updating %i dirty actor meshes"), *GetOwner()->GetActorNameOrLabel(), DirtyAttachments.Num());
 
 	// Update baked dynamic meshes on other thread
 	TWeakObjectPtr<UCEClonerComponent> ThisWeak(this);
@@ -803,7 +282,7 @@ void UCEClonerComponent::UpdateDirtyMeshesAsync()
 				continue;
 			}
 
-			This->UpdateActorBakedDynamicMesh(DirtyActor);
+			This->UpdateActorDynamicMesh(DirtyActor);
 		}
 
 		// Create baked static mesh on main thread (required)
@@ -825,18 +304,18 @@ void UCEClonerComponent::UpdateDirtyMeshesAsync()
 			// Update actors baked static mesh
 			for (int32 Idx = 0; Idx < This->ClonerTree.RootActors.Num(); Idx++)
 			{
-				if (IsGarbageCollectingAndLockingUObjectHashTables())
-				{
-					bSuccess = false;
-					break;
-				}
-
 				const UStaticMesh* RootStaticMesh = This->ClonerTree.MergedBakedMeshes[Idx].Get();
 
 				if (!RootStaticMesh)
 				{
+					if (IsGarbageCollectingAndLockingUObjectHashTables())
+					{
+						bSuccess = false;
+						break;
+					}
+
 					AActor* RootActor = This->ClonerTree.RootActors[Idx].Get();
-					This->UpdateRootActorBakedStaticMesh(RootActor);
+					This->UpdateActorStaticMesh(RootActor);
 				}
 			}
 
@@ -846,18 +325,7 @@ void UCEClonerComponent::UpdateDirtyMeshesAsync()
 	});
 }
 
-void UCEClonerComponent::OnDirtyMeshesUpdated(bool bInSuccess)
-{
-	bClonerMeshesUpdating = false;
-
-	// Update niagara parameters
-	if (bInSuccess)
-	{
-		UpdateClonerMeshes();
-	}
-}
-
-void UCEClonerComponent::UpdateActorBakedDynamicMesh(AActor* InActor)
+void UCEClonerComponent::UpdateActorDynamicMesh(AActor* InActor)
 {
 	if (!InActor)
 	{
@@ -880,7 +348,7 @@ void UCEClonerComponent::UpdateActorBakedDynamicMesh(AActor* InActor)
 
 	AttachmentItem->MeshStatus = ECEClonerAttachmentStatus::Updating;
 
-	UE_LOG(LogCEClonerComponent, Log, TEXT("%s : Updating baked actor mesh %s"), *ClonerActor->GetActorNameOrLabel(), *InActor->GetActorNameOrLabel());
+	UE_LOG(LogCECloner, Log, TEXT("%s : Updating actor dynamic mesh %s"), *ClonerActor->GetActorNameOrLabel(), *InActor->GetActorNameOrLabel());
 
 	UDynamicMesh* Mesh = NewObject<UDynamicMesh>();
 	TArray<TWeakObjectPtr<UMaterialInterface>> MeshMaterials;
@@ -893,20 +361,25 @@ void UCEClonerComponent::UpdateActorBakedDynamicMesh(AActor* InActor)
 
 	TArray<TWeakObjectPtr<UMaterialInterface>> UnsetMaterials;
 	UMaterialInterface* DefaultMaterial = LoadObject<UMaterialInterface>(nullptr, UCEClonerEffectorSettings::DefaultMaterialPath);
-	if (FilterSupportedMaterials(MeshMaterials, UnsetMaterials, DefaultMaterial))
+	if (UE::ClonerEffector::Utilities::FilterSupportedMaterials(MeshMaterials, UnsetMaterials, DefaultMaterial))
 	{
-		FireMaterialWarning(InActor, UnsetMaterials);
+		if (UCEClonerSubsystem* ClonerSubsystem = UCEClonerSubsystem::Get())
+		{
+			ClonerSubsystem->FireMaterialWarning(ClonerActor, InActor, UnsetMaterials);	
+		}
 	}
 
 	AttachmentItem->BakedMaterials = MoveTemp(MeshMaterials);
 
-	// Was the mesh invalidated during the update process ?
-	AttachmentItem->MeshStatus = AttachmentItem->MeshStatus == ECEClonerAttachmentStatus::Outdated ? ECEClonerAttachmentStatus::Outdated : ECEClonerAttachmentStatus::Updated; //-V547
-
-	InvalidateBakedStaticMesh(AttachmentItem->ItemActor.Get());
+	// Was the mesh invalidated during the update process, then leave it outdated
+	if (AttachmentItem->MeshStatus == ECEClonerAttachmentStatus::Updating)
+	{
+		AttachmentItem->MeshStatus = ECEClonerAttachmentStatus::Updated;
+		ClonerTree.MarkCacheOutdated(InActor);
+	}
 }
 
-void UCEClonerComponent::UpdateRootActorBakedStaticMesh(AActor* InRootActor)
+void UCEClonerComponent::UpdateActorStaticMesh(AActor* InRootActor)
 {
 	if (!InRootActor)
 	{
@@ -934,12 +407,12 @@ void UCEClonerComponent::UpdateRootActorBakedStaticMesh(AActor* InRootActor)
 		return;
 	}
 
-	UE_LOG(LogCEClonerComponent, Log, TEXT("%s : Updating root merged baked mesh %s"), *ClonerActor->GetActorNameOrLabel(), *InRootActor->GetActorNameOrLabel());
+	UE_LOG(LogCECloner, Log, TEXT("%s : Updating actor static mesh %s"), *ClonerActor->GetActorNameOrLabel(), *InRootActor->GetActorNameOrLabel());
 
-	TArray<FCEClonerAttachmentItem*> AttachmentItems;
-	GetActorAttachmentItems(InRootActor, AttachmentItems);
+	TArray<const FCEClonerAttachmentItem*> AttachmentItems;
+	ClonerTree.GetAttachments(InRootActor, AttachmentItems, /** Recurse */true);
 
-	for (FCEClonerAttachmentItem* AttachmentItem : AttachmentItems)
+	for (const FCEClonerAttachmentItem* AttachmentItem : AttachmentItems)
 	{
 		if (!AttachmentItem)
 		{
@@ -952,164 +425,44 @@ void UCEClonerComponent::UpdateRootActorBakedStaticMesh(AActor* InRootActor)
 			continue;
 		}
 
-		FTransform MeshTransform = FTransform::Identity;
-		if (AttachmentItem->ParentActor.IsValid())
-		{
-			const FTransform& ParentTransform = AttachmentItem->ParentActor->GetTransform();
-			MeshTransform = AttachmentItem->ItemActor->GetTransform().GetRelativeTransform(ParentTransform);
-		}
-
-		MeshBuilder.AppendMesh(AttachmentItem->BakedMesh, AttachmentItem->BakedMaterials, MeshTransform);
+		MeshBuilder.AppendMesh(AttachmentItem->BakedMesh, AttachmentItem->BakedMaterials, AttachmentItem->ActorTransform);
 	}
 
-	// Avoid dirtying asset
-	UStaticMesh* Mesh = NewObject<UStaticMesh>();
+	// Avoid dirtying asset by creating it transient first,
+	// and avoid bounds log spamming by renaming asset with prefix LandscapeNaniteMesh
+	const FName PreObjectName(FString::Printf(TEXT("LandscapeNaniteMesh_C%i_M%i"), ClonerActor->GetUniqueID(), InRootActor->GetUniqueID()));
+	UStaticMesh* Mesh = NewObject<UStaticMesh>(GetTransientPackage(), PreObjectName);
+
+	// Rename the asset + outer once build is done
+	Mesh->OnPostMeshBuild().AddUObject(this, &UCEClonerComponent::OnActorStaticMeshPostBuild);
 
 	TArray<TWeakObjectPtr<UMaterialInterface>> MeshMaterials;
-	bClonerMeshesDirty = MeshBuilder.BuildStaticMesh(Mesh, MeshMaterials);
+	ClonerTree.bItemAttachmentsDirty = MeshBuilder.BuildStaticMesh(Mesh, MeshMaterials);
 	MeshBuilder.Reset();
 
-	Mesh->Rename(nullptr, this, REN_NonTransactional | REN_DoNotDirty);
-
 	ClonerTree.MergedBakedMeshes[RootIdx] = Mesh;
 }
 
-bool UCEClonerComponent::FilterSupportedMaterials(TArray<TWeakObjectPtr<UMaterialInterface>>& InMaterials, TArray<TWeakObjectPtr<UMaterialInterface>>& OutUnsetMaterials, UMaterialInterface* InDefaultMaterial)
-{
-	check(InDefaultMaterial)
-
-	OutUnsetMaterials.Reset(InMaterials.Num());
-
-	for (int32 Index = 0; Index < InMaterials.Num(); Index++)
-	{
-		UMaterialInterface* PreviousMaterialInterface = InMaterials[Index].Get();
-
-		UMaterialInterface* NewMaterialInterface = PreviousMaterialInterface;
-
-		if (FilterSupportedMaterial(NewMaterialInterface, InDefaultMaterial))
-		{
-			// Add original material to unset list
-			OutUnsetMaterials.Add(PreviousMaterialInterface);
-		}
-
-		// Replace material
-		InMaterials[Index] = NewMaterialInterface;
-	}
-
-	return OutUnsetMaterials.IsEmpty();
-}
-
-bool UCEClonerComponent::FilterSupportedMaterial(UMaterialInterface*& InMaterial, UMaterialInterface* InDefaultMaterial)
-{
-	if (InMaterial && !IsMaterialUsageFlagSet(InMaterial))
-	{
-		// Replace material if dirtyable and not in read only location
-		if (!IsMaterialDirtyable(InMaterial))
-		{
-			InMaterial = InDefaultMaterial;
-		}
-
-		return true;
-	}
-
-	return false;
-}
-
-void UCEClonerComponent::GetActorAttachmentItems(AActor* InActor, TArray<FCEClonerAttachmentItem*>& OutAttachmentItems)
-{
-	if (!InActor)
-	{
-		return;
-	}
-
-	FCEClonerAttachmentItem* AttachmentItem = ClonerTree.ItemAttachmentMap.Find(InActor);
-
-	if (!AttachmentItem)
-	{
-		return;
-	}
-
-	OutAttachmentItems.Add(AttachmentItem);
-
-	for (TWeakObjectPtr<AActor>& ChildActor : AttachmentItem->ChildrenActors)
-	{
-		if (ChildActor.IsValid())
-		{
-			GetActorAttachmentItems(ChildActor.Get(), OutAttachmentItems);
-		}
-	}
-}
-
-bool UCEClonerComponent::IsAllMergedMeshesValid() const
-{
-	for (const TObjectPtr<UStaticMesh>& MergedMesh : ClonerTree.MergedBakedMeshes)
-	{
-		if (!MergedMesh.Get())
-		{
-			return false;
-		}
-	}
-	return true;
-}
-
-bool UCEClonerComponent::IsMaterialDirtyable(const UMaterialInterface* InMaterial)
-{
-	const UMaterial* BaseMaterial = InMaterial->GetMaterial_Concurrent();
-	const FString ContentFolder = FPaths::ConvertRelativePathToFull(FPaths::ProjectContentDir());
-
-	const UPackage* MaterialPackage = BaseMaterial->GetPackage();
-	const FPackagePath& LoadedPath = MaterialPackage->GetLoadedPath();
-	const FString PackagePath = FPaths::ConvertRelativePathToFull(LoadedPath.GetLocalFullPath());
-	const FString MaterialPath = BaseMaterial->GetPathName();
-
-	const bool bTransientPackage = MaterialPackage == GetTransientPackage() || MaterialPath.StartsWith("/Temp/");
-	const bool bContentFolder = PackagePath.StartsWith(ContentFolder);
-
-	return bTransientPackage || bContentFolder;
-}
-
-bool UCEClonerComponent::IsMaterialUsageFlagSet(const UMaterialInterface* InMaterial)
+void UCEClonerComponent::OnActorStaticMeshPostBuild(UStaticMesh* InMesh)
 {
-	if (InMaterial)
+	if (InMesh)
 	{
-		if (const UMaterial* Material = InMaterial->GetMaterial_Concurrent())
-		{
-			return Material->GetUsageByFlag(EMaterialUsage::MATUSAGE_NiagaraMeshParticles);
-		}
+		InMesh->OnPostMeshBuild().RemoveAll(this);
+		const FName MeshName = MakeUniqueObjectName(this, UStaticMesh::StaticClass(), TEXT("ClonerMesh"));
+		InMesh->Rename(*MeshName.ToString(), this, REN_NonTransactional | REN_DoNotDirty);
 	}
-
-	return false;
 }
+#endif // WITH_EDITOR
 
-void UCEClonerComponent::FireMaterialWarning(const AActor* InContextActor, const TArray<TWeakObjectPtr<UMaterialInterface>>& InUnsetMaterials)
+void UCEClonerComponent::OnDirtyMeshesUpdated(bool bInSuccess)
 {
-	if (!IsValid(InContextActor) || InUnsetMaterials.IsEmpty())
-	{
-		return;
-	}
-
-	UE_LOG(LogCEClonerComponent, Warning, TEXT("%s : %i unsupported material(s) detected due to missing niagara usage flag (bUsedWithNiagaraMeshParticles) on actor (%s), see logs below"), *GetOwner()->GetActorNameOrLabel(), InUnsetMaterials.Num(), *InContextActor->GetActorNameOrLabel());
-
-	const AActor* ClonerActor = GetOwner();
-	for (const TWeakObjectPtr<UMaterialInterface>& UnsetMaterialWeak : InUnsetMaterials)
-	{
-		if (UMaterialInterface* UnsetMaterial = UnsetMaterialWeak.Get())
-		{
-			UE_LOG(LogCEClonerComponent, Warning, TEXT("%s : The following materials (%s) on actor (%s) does not have the usage flag (bUsedWithNiagaraMeshParticles) set to work with the cloner, set the flag and resave the asset to avoid this warning"), *ClonerActor->GetActorNameOrLabel(), *UnsetMaterial->GetMaterial()->GetPathName(), *InContextActor->GetActorNameOrLabel());
-		}
-	}
-
-#if WITH_EDITOR
-	// Fire warning notification when invalid materials are found and at least 5s has elapsed since last notification
-	constexpr double MinNotificationElapsedTime = 5.0;
-	const double CurrentTime = FApp::GetCurrentTime();
+	bClonerMeshesUpdating = false;
 
-	if (CurrentTime - LastNotificationTime > MinNotificationElapsedTime)
+	// Update niagara parameters
+	if (bInSuccess)
 	{
-		LastNotificationTime = CurrentTime;
-		ShowMaterialWarning(InUnsetMaterials.Num());
+		UpdateClonerMeshes();
 	}
-#endif
 }
 
 void UCEClonerComponent::GetPreloadDependencies(TArray<UObject*>& OutDeps)
@@ -1125,39 +478,12 @@ void UCEClonerComponent::GetPreloadDependencies(TArray<UObject*>& OutDeps)
 	}
 }
 
-#if WITH_EDITOR
-void UCEClonerComponent::ShowMaterialWarning(int32 InMaterialCount)
-{
-	if (InMaterialCount > 0)
-	{
-		FNotificationInfo NotificationInfo(FText::Format(LOCTEXT("MaterialsMissingUsageFlag", "Detected {0} material(s) with missing usage flag required to work properly with cloner (See logs)"), InMaterialCount));
-		NotificationInfo.ExpireDuration = 5.f;
-		NotificationInfo.bFireAndForget = true;
-		NotificationInfo.Image = FAppStyle::GetBrush("Icons.WarningWithColor");
-
-		FSlateNotificationManager::Get().AddNotification(NotificationInfo);
-	}
-}
-
-FName UCEClonerComponent::GetActiveExtensionsPropertyName()
-{
-	return GET_MEMBER_NAME_CHECKED(UCEClonerComponent, ActiveExtensions);
-}
-
-FName UCEClonerComponent::GetActiveLayoutPropertyName()
-{
-	return GET_MEMBER_NAME_CHECKED(UCEClonerComponent, ActiveLayout);
-}
-
-FName UCEClonerComponent::GetLayoutNamePropertyName()
-{
-	return GET_MEMBER_NAME_CHECKED(UCEClonerComponent, LayoutName);
-}
-#endif // WITH_EDITOR
-
 void UCEClonerComponent::InitializeCloner()
 {
-	if (IsTemplate() || bClonerInitialized)
+	AActor* Owner = GetOwner();
+	if (bClonerInitialized
+		|| IsTemplate()
+		|| !Owner)
 	{
 		return;
 	}
@@ -1169,15 +495,16 @@ void UCEClonerComponent::InitializeCloner()
 #if WITH_EDITOR
 	OnVisualizerSpriteVisibleChanged();
 
-	const AActor* Owner = GetOwner();
-
 	// Skip init for preview actor
-	if (Owner && Owner->bIsEditorPreviewActor)
+	if (Owner->bIsEditorPreviewActor)
 	{
 		return;
 	}
 #endif
 
+	ClonerTree.SetAttachmentRoot(Owner);
+	OnTreeBehaviorNameChanged();
+
 	// When level is streamed in, wait until actor hierarchy and resources are ready before initializing
 	if (const ULevel* Level = GetComponentLevel())
 	{
@@ -1209,10 +536,7 @@ void UCEClonerComponent::OnLevelLoaded(ULevel* InLevel, UWorld* InWorld)
 void UCEClonerComponent::PostInitializeCloner()
 {
 	// Register a custom ticker to avoid using the component tick that needs the simulation to be solo
-	TreeUpdateDeltaTime = TreeUpdateInterval;
-
 	RegisterTicker();
-
 	OnClonerInitializedDelegate.Broadcast(this);
 }
 
@@ -1239,12 +563,9 @@ bool UCEClonerComponent::CheckResourcesReady()
 	}
 
 	// Check cached meshes are not async loading
-	for (const TObjectPtr<UStaticMesh>& BakedMesh : ClonerTree.MergedBakedMeshes)
+	if (!ClonerTree.IsCacheAvailable(/** AllowInvalid */true))
 	{
-		if (BakedMesh && BakedMesh->HasAnyInternalFlags(EInternalObjectFlags_AsyncLoading))
-		{
-			return false;
-		}
+		return false;
 	}
 
 	// Check cached active system is not async loading
@@ -1282,15 +603,9 @@ bool UCEClonerComponent::TickCloner(float InDelta)
 
 	if (bEnabled)
 	{
-		TreeUpdateDeltaTime += InDelta;
-
 		// Update attachment tree
-		if (TreeUpdateDeltaTime >= TreeUpdateInterval)
-		{
-			TreeUpdateDeltaTime -= TreeUpdateInterval != 0.f ? TreeUpdateInterval : TreeUpdateDeltaTime;
-			UpdateClonerAttachmentTree();
-			UpdateClonerRenderState();
-		}
+		ClonerTree.UpdateAttachments();
+		UpdateClonerRenderState();
 
 		// Update layout parameters
 		if (ActiveLayout && ActiveLayout->IsLayoutDirty())
@@ -1329,16 +644,6 @@ void UCEClonerComponent::SetEnabled(bool bInEnable)
 	OnEnabledChanged();
 }
 
-void UCEClonerComponent::SetTreeUpdateInterval(float InInterval)
-{
-	if (InInterval == TreeUpdateInterval)
-	{
-		return;
-	}
-
-	TreeUpdateInterval = InInterval;
-}
-
 void UCEClonerComponent::SetSeed(int32 InSeed)
 {
 	if (InSeed == Seed)
@@ -1424,6 +729,23 @@ TSubclassOf<UCEClonerLayoutBase> UCEClonerComponent::GetLayoutClass() const
 }
 
 #if WITH_EDITOR
+void UCEClonerComponent::SetTreeBehaviorName(FName InBehaviorName)
+{
+	if (TreeBehaviorName == InBehaviorName)
+	{
+		return;
+	}
+
+	const TArray<FName> BehaviorNames = GetClonerTreeBehaviorNames();
+	if (!BehaviorNames.Contains(InBehaviorName))
+	{
+		return;
+	}
+
+	TreeBehaviorName = InBehaviorName;
+	OnTreeBehaviorNameChanged();
+}
+
 void UCEClonerComponent::SetVisualizerSpriteVisible(bool bInVisible)
 {
 	if (bVisualizerSpriteVisible == bInVisible)
@@ -1457,7 +779,8 @@ int32 UCEClonerComponent::GetAttachmentCount() const
 #if WITH_EDITOR
 void UCEClonerComponent::ForceUpdateCloner()
 {
-	UpdateClonerAttachmentTree();
+	ClonerTree.MarkAttachmentOutdated();
+	ClonerTree.UpdateAttachments();
 	UpdateClonerRenderState();
 	OnLayoutNameChanged();
 }
@@ -1536,9 +859,7 @@ void UCEClonerComponent::CreateDefaultActorAttached()
 
 void UCEClonerComponent::ConvertToStaticMesh()
 {
-	PRAGMA_DISABLE_DEPRECATION_WARNINGS
-	if (!IsThisNotNull(this, "UCEClonerComponent::ConvertToStaticMesh") || !IsValidChecked(this) || !bEnabled)
-	PRAGMA_ENABLE_DEPRECATION_WARNINGS
+	if (!IsValid(this) || !bEnabled)
 	{
 		return;
 	}
@@ -1546,23 +867,21 @@ void UCEClonerComponent::ConvertToStaticMesh()
 	FScopedSlowTask SlowTask(0.0f, LOCTEXT("ConvertToStaticMesh", "Converting cloner to static mesh"));
 	SlowTask.MakeDialog();
 
-	UE_LOG(LogCEClonerComponent, Log, TEXT("%s : Request ConvertToStaticMesh..."), *GetOwner()->GetActorNameOrLabel())
+	UE_LOG(LogCECloner, Log, TEXT("%s : Request ConvertToStaticMesh..."), *GetOwner()->GetActorNameOrLabel())
 
 	if (UE::ClonerEffector::Conversion::ConvertClonerToStaticMesh(this))
 	{
-		UE_LOG(LogCEClonerComponent, Log, TEXT("%s : ConvertToStaticMesh Completed"), *GetOwner()->GetActorNameOrLabel())
+		UE_LOG(LogCECloner, Log, TEXT("%s : ConvertToStaticMesh Completed"), *GetOwner()->GetActorNameOrLabel())
 	}
 	else
 	{
-		UE_LOG(LogCEClonerComponent, Warning, TEXT("%s : ConvertToStaticMesh Failed"), *GetOwner()->GetActorNameOrLabel())
+		UE_LOG(LogCECloner, Warning, TEXT("%s : ConvertToStaticMesh Failed"), *GetOwner()->GetActorNameOrLabel())
 	}
 }
 
 void UCEClonerComponent::ConvertToDynamicMesh()
 {
-	PRAGMA_DISABLE_DEPRECATION_WARNINGS
-	if (!IsThisNotNull(this, "UCEClonerComponent::ConvertToDynamicMesh") || !IsValidChecked(this) || !bEnabled)
-	PRAGMA_ENABLE_DEPRECATION_WARNINGS
+	if (!IsValid(this) || !bEnabled)
 	{
 		return;
 	}
@@ -1570,23 +889,21 @@ void UCEClonerComponent::ConvertToDynamicMesh()
 	FScopedSlowTask SlowTask(0.0f, LOCTEXT("ConvertToDynamicMesh", "Converting cloner to dynamic mesh"));
 	SlowTask.MakeDialog();
 
-	UE_LOG(LogCEClonerComponent, Log, TEXT("%s : Request ConvertToDynamicMesh..."), *GetOwner()->GetActorNameOrLabel())
+	UE_LOG(LogCECloner, Log, TEXT("%s : Request ConvertToDynamicMesh..."), *GetOwner()->GetActorNameOrLabel())
 
 	if (UE::ClonerEffector::Conversion::ConvertClonerToDynamicMesh(this))
 	{
-		UE_LOG(LogCEClonerComponent, Log, TEXT("%s : ConvertToDynamicMesh Completed"), *GetOwner()->GetActorNameOrLabel())
+		UE_LOG(LogCECloner, Log, TEXT("%s : ConvertToDynamicMesh Completed"), *GetOwner()->GetActorNameOrLabel())
 	}
 	else
 	{
-		UE_LOG(LogCEClonerComponent, Warning, TEXT("%s : ConvertToDynamicMesh Failed"), *GetOwner()->GetActorNameOrLabel())
+		UE_LOG(LogCECloner, Warning, TEXT("%s : ConvertToDynamicMesh Failed"), *GetOwner()->GetActorNameOrLabel())
 	}
 }
 
 void UCEClonerComponent::ConvertToStaticMeshes()
 {
-	PRAGMA_DISABLE_DEPRECATION_WARNINGS
-	if (!IsThisNotNull(this, "UCEClonerComponent::ConvertToStaticMeshes") || !IsValidChecked(this) || !bEnabled)
-	PRAGMA_ENABLE_DEPRECATION_WARNINGS
+	if (!IsValid(this) || !bEnabled)
 	{
 		return;
 	}
@@ -1594,23 +911,21 @@ void UCEClonerComponent::ConvertToStaticMeshes()
 	FScopedSlowTask SlowTask(0.0f, LOCTEXT("ConvertToStaticMeshes", "Converting cloner to static meshes"));
 	SlowTask.MakeDialog();
 
-	UE_LOG(LogCEClonerComponent, Log, TEXT("%s : Request ConvertToStaticMeshes..."), *GetOwner()->GetActorNameOrLabel())
+	UE_LOG(LogCECloner, Log, TEXT("%s : Request ConvertToStaticMeshes..."), *GetOwner()->GetActorNameOrLabel())
 
 	if (!UE::ClonerEffector::Conversion::ConvertClonerToStaticMeshes(this).IsEmpty())
 	{
-		UE_LOG(LogCEClonerComponent, Log, TEXT("%s : ConvertToStaticMeshes Completed"), *GetOwner()->GetActorNameOrLabel())
+		UE_LOG(LogCECloner, Log, TEXT("%s : ConvertToStaticMeshes Completed"), *GetOwner()->GetActorNameOrLabel())
 	}
 	else
 	{
-		UE_LOG(LogCEClonerComponent, Warning, TEXT("%s : ConvertToStaticMeshes Failed"), *GetOwner()->GetActorNameOrLabel())
+		UE_LOG(LogCECloner, Warning, TEXT("%s : ConvertToStaticMeshes Failed"), *GetOwner()->GetActorNameOrLabel())
 	}
 }
 
 void UCEClonerComponent::ConvertToDynamicMeshes()
 {
-	PRAGMA_DISABLE_DEPRECATION_WARNINGS
-	if (!IsThisNotNull(this, "UCEClonerComponent::ConvertToDynamicMeshes") || !IsValidChecked(this) || !bEnabled)
-	PRAGMA_ENABLE_DEPRECATION_WARNINGS
+	if (!IsValid(this) || !bEnabled)
 	{
 		return;
 	}
@@ -1618,23 +933,21 @@ void UCEClonerComponent::ConvertToDynamicMeshes()
 	FScopedSlowTask SlowTask(0.0f, LOCTEXT("ConvertToDynamicMeshes", "Converting cloner to dynamic meshes"));
 	SlowTask.MakeDialog();
 
-	UE_LOG(LogCEClonerComponent, Log, TEXT("%s : Request ConvertToDynamicMeshes..."), *GetOwner()->GetActorNameOrLabel())
+	UE_LOG(LogCECloner, Log, TEXT("%s : Request ConvertToDynamicMeshes..."), *GetOwner()->GetActorNameOrLabel())
 
 	if (!UE::ClonerEffector::Conversion::ConvertClonerToDynamicMeshes(this).IsEmpty())
 	{
-		UE_LOG(LogCEClonerComponent, Log, TEXT("%s : ConvertToDynamicMeshes Completed"), *GetOwner()->GetActorNameOrLabel())
+		UE_LOG(LogCECloner, Log, TEXT("%s : ConvertToDynamicMeshes Completed"), *GetOwner()->GetActorNameOrLabel())
 	}
 	else
 	{
-		UE_LOG(LogCEClonerComponent, Warning, TEXT("%s : ConvertToDynamicMeshes Failed"), *GetOwner()->GetActorNameOrLabel())
+		UE_LOG(LogCECloner, Warning, TEXT("%s : ConvertToDynamicMeshes Failed"), *GetOwner()->GetActorNameOrLabel())
 	}
 }
 
 void UCEClonerComponent::ConvertToInstancedStaticMeshes()
 {
-	PRAGMA_DISABLE_DEPRECATION_WARNINGS
-	if (!IsThisNotNull(this, "UCEClonerComponent::ConvertToInstancedStaticMeshes") || !IsValidChecked(this) || !bEnabled)
-	PRAGMA_ENABLE_DEPRECATION_WARNINGS
+	if (!IsValid(this) || !bEnabled)
 	{
 		return;
 	}
@@ -1642,15 +955,15 @@ void UCEClonerComponent::ConvertToInstancedStaticMeshes()
 	FScopedSlowTask SlowTask(0.0f, LOCTEXT("ConvertToInstancedStaticMeshes", "Converting cloner to instanced static meshes"));
 	SlowTask.MakeDialog();
 
-	UE_LOG(LogCEClonerComponent, Log, TEXT("%s : Request ConvertToInstancedStaticMeshes..."), *GetOwner()->GetActorNameOrLabel())
+	UE_LOG(LogCECloner, Log, TEXT("%s : Request ConvertToInstancedStaticMeshes..."), *GetOwner()->GetActorNameOrLabel())
 
 	if (!UE::ClonerEffector::Conversion::ConvertClonerToInstancedStaticMeshes(this).IsEmpty())
 	{
-		UE_LOG(LogCEClonerComponent, Log, TEXT("%s : ConvertToInstancedStaticMeshes Completed"), *GetOwner()->GetActorNameOrLabel())
+		UE_LOG(LogCECloner, Log, TEXT("%s : ConvertToInstancedStaticMeshes Completed"), *GetOwner()->GetActorNameOrLabel())
 	}
 	else
 	{
-		UE_LOG(LogCEClonerComponent, Warning, TEXT("%s : ConvertToInstancedStaticMeshes Failed"), *GetOwner()->GetActorNameOrLabel())
+		UE_LOG(LogCECloner, Warning, TEXT("%s : ConvertToInstancedStaticMeshes Failed"), *GetOwner()->GetActorNameOrLabel())
 	}
 }
 #endif
@@ -1743,12 +1056,31 @@ void UCEClonerComponent::OnColorChanged()
 void UCEClonerComponent::OnGlobalScaleChanged()
 {
 	GlobalScale = GlobalScale.ComponentMax(FVector(UE_KINDA_SMALL_NUMBER));
-	bClonerMeshesDirty = true;
+	ClonerTree.bItemAttachmentsDirty = true;
 }
 
 void UCEClonerComponent::OnGlobalRotationChanged()
 {
-	bClonerMeshesDirty = true;
+	ClonerTree.bItemAttachmentsDirty = true;
+}
+
+void UCEClonerComponent::OnTreeBehaviorNameChanged()
+{
+	const TArray<FName> BehaviorNames = GetClonerTreeBehaviorNames();
+
+	// Set default if value does not exists
+	if (!BehaviorNames.Contains(TreeBehaviorName) && !BehaviorNames.IsEmpty())
+	{
+		TreeBehaviorName = BehaviorNames[0];
+	}
+
+	if (const UCEClonerSubsystem* ClonerSubsystem = UCEClonerSubsystem::Get())
+	{
+		if (const TSharedPtr<ICEClonerAttachmentTreeBehavior> TreeBehavior = ClonerSubsystem->CreateAttachmentTreeBehavior(TreeBehaviorName))
+		{
+			ClonerTree.SetBehaviorImplementation(TreeBehavior.ToSharedRef());
+		}
+	}
 }
 
 void UCEClonerComponent::OnLayoutNameChanged()
@@ -1790,88 +1122,24 @@ void UCEClonerComponent::OnVisualizerSpriteVisibleChanged()
 		}
 	}
 }
+#endif
 
-void UCEClonerComponent::OnRenderStateDirty(UActorComponent& InActorComponent)
+void UCEClonerComponent::OnTreeItemAttached(AActor* InActor, FCEClonerAttachmentItem& InItem)
 {
-	AActor* Owner = InActorComponent.GetOwner();
-	const AActor* ClonerActor = GetOwner();
-
-	if (!Owner || Owner->GetLevel() != ClonerActor->GetLevel())
-	{
-		return;
-	}
-
-	// Does it contain geometry that we can convert
-	if (!FCEMeshBuilder::IsComponentSupported(&InActorComponent))
-	{
-		return;
-	}
-
-	FCEClonerAttachmentItem* Item = ClonerTree.ItemAttachmentMap.Find(Owner);
-
-	if (!Item || !Item->CheckBoundsChanged(/** CheckAndUpdate */false))
+	if (InActor)
 	{
-		return;
+		UE::ClonerEffector::Utilities::SetActorVisibility(InActor, /** Visible */false);
+		OnClonerActorAttachedDelegate.Broadcast(this, InActor);
 	}
-
-	UE_LOG(LogCEClonerComponent, Log, TEXT("%s : Render state changed for %s"), *ClonerActor->GetActorNameOrLabel(), *Owner->GetActorNameOrLabel());
-
-	// Rebind delegates as new components might be available
-	BindActorDelegates(Owner);
-
-	Item->MeshStatus = ECEClonerAttachmentStatus::Outdated;
-	InvalidateBakedStaticMesh(Owner);
-	ClonerTree.DirtyItemAttachments.Add(Item->ItemActor);
 }
-#endif
 
-void UCEClonerComponent::OnComponentTransformed(USceneComponent* InComponent, EUpdateTransformFlags InFlags, ETeleportType InTeleport)
+void UCEClonerComponent::OnTreeItemDetached(AActor* InActor, FCEClonerAttachmentItem& InItem)
 {
-	if (!InComponent || !InComponent->GetOwner() || InFlags == EUpdateTransformFlags::PropagateFromParent)
-	{
-		return;
-	}
-
-	AActor* Owner = InComponent->GetOwner();
-	const AActor* RootActor = GetRootActor(Owner);
-
-	// Skip update if root component has moved, since we can simply offset the mesh
-	if (!RootActor || (RootActor == Owner && RootActor->GetRootComponent() == InComponent))
-	{
-		return;
-	}
-
-	bool bComponentSupported = FCEMeshBuilder::IsComponentSupported(InComponent);
-
-	if (!bComponentSupported)
+	if (InActor)
 	{
-		for (const TObjectPtr<USceneComponent>& ChildComponent : InComponent->GetAttachChildren())
-		{
-			if (FCEMeshBuilder::IsComponentSupported(ChildComponent))
-			{
-				bComponentSupported = true;
-				break;
-			}
-		}
-	}
-
-	if (!bComponentSupported)
-	{
-		return;
+		UE::ClonerEffector::Utilities::SetActorVisibility(InActor, /** Visible */true);
+		OnClonerActorDetachedDelegate.Broadcast(this, InActor);
 	}
-
-	FCEClonerAttachmentItem* Item = ClonerTree.ItemAttachmentMap.Find(Owner);
-
-	if (!Item)
-	{
-		return;
-	}
-
-	UE_LOG(LogCEClonerComponent, Log, TEXT("%s : Transform state changed for %s"), *GetOwner()->GetActorNameOrLabel(), *Owner->GetActorNameOrLabel());
-
-	Item->MeshStatus = ECEClonerAttachmentStatus::Outdated;
-	InvalidateBakedStaticMesh(Owner);
-	ClonerTree.DirtyItemAttachments.Add(Item->ItemActor);
 }
 
 UCEClonerLayoutBase* UCEClonerComponent::FindOrAddLayout(TSubclassOf<UCEClonerLayoutBase> InClass)
@@ -1991,9 +1259,21 @@ TArray<FName> UCEClonerComponent::GetClonerLayoutNames() const
 	return LayoutNames;
 }
 
+TArray<FName> UCEClonerComponent::GetClonerTreeBehaviorNames() const
+{
+	TArray<FName> LayoutNames;
+
+	if (const UCEClonerSubsystem* Subsystem = UCEClonerSubsystem::Get())
+	{
+		LayoutNames = Subsystem->GetAttachmentTreeBehaviorNames();
+	}
+
+	return LayoutNames;
+}
+
 void UCEClonerComponent::RefreshClonerMeshes()
 {
-	if (!bClonerMeshesUpdating && !bClonerMeshesDirty)
+	if (!bClonerMeshesUpdating && !ClonerTree.bItemAttachmentsDirty)
 	{
 		UpdateClonerMeshes();
 	}
@@ -2033,12 +1313,7 @@ UCEClonerExtensionBase* UCEClonerComponent::GetExtension(FName InExtensionName)
 
 TArray<AActor*> UCEClonerComponent::GetClonerRootActors() const
 {
-	TArray<AActor*> RootActors;
-	Algo::Transform(ClonerTree.RootActors, RootActors, [](const TWeakObjectPtr<AActor>& InActorWeak)->AActor*
-	{
-		return InActorWeak.Get();
-	});
-	return RootActors;
+	return ClonerTree.GetRootActors();
 }
 
 void UCEClonerComponent::OnActiveLayoutLoaded(UCEClonerLayoutBase* InLayout, bool bInSuccess)
@@ -2052,11 +1327,11 @@ void UCEClonerComponent::OnActiveLayoutLoaded(UCEClonerLayoutBase* InLayout, boo
 
 	if (!bInSuccess)
 	{
-		UE_LOG(LogCEClonerComponent, Warning, TEXT("%s : Cloner layout system failed to load %s - %s"), *GetOwner()->GetActorNameOrLabel(), *InLayout->GetLayoutName().ToString(), *InLayout->GetLayoutAssetPath());
+		UE_LOG(LogCECloner, Warning, TEXT("%s : Cloner layout system failed to load %s - %s"), *GetOwner()->GetActorNameOrLabel(), *InLayout->GetLayoutName().ToString(), *InLayout->GetLayoutAssetPath());
 		return;
 	}
 
-	UE_LOG(LogCEClonerComponent, Log, TEXT("%s : Cloner layout system loaded %s - %s"), *GetOwner()->GetActorNameOrLabel(), *InLayout->GetLayoutName().ToString(), *InLayout->GetLayoutAssetPath());
+	UE_LOG(LogCECloner, Log, TEXT("%s : Cloner layout system loaded %s - %s"), *GetOwner()->GetActorNameOrLabel(), *InLayout->GetLayoutName().ToString(), *InLayout->GetLayoutAssetPath());
 
 	OnClonerLayoutLoadedDelegate.Broadcast(this, InLayout);
 
@@ -2088,11 +1363,11 @@ void UCEClonerComponent::ActivateLayout(UCEClonerLayoutBase* InLayout)
 
 	ActiveLayout = InLayout;
 
-	UE_LOG(LogCEClonerComponent, Log, TEXT("%s : Cloner layout system changed %s - %s"), *GetOwner()->GetActorNameOrLabel(), *InLayout->GetLayoutName().ToString(), *InLayout->GetLayoutAssetPath());
+	UE_LOG(LogCECloner, Log, TEXT("%s : Cloner layout system changed %s - %s"), *GetOwner()->GetActorNameOrLabel(), *InLayout->GetLayoutName().ToString(), *InLayout->GetLayoutAssetPath());
 
 	OnActiveLayoutChanged();
 
-	bClonerMeshesDirty = true;
+	ClonerTree.bItemAttachmentsDirty = true;
 }
 
 void UCEClonerComponent::OnActiveLayoutChanged()
@@ -2156,7 +1431,7 @@ void UCEClonerComponent::UpdateClonerMeshes()
 
 	if (ActiveLayout->GetSystem() != ActiveSystem)
 	{
-		UE_LOG(LogCEClonerComponent, Warning, TEXT("%s : Invalid system for cloner layout"), *ClonerActor->GetActorNameOrLabel());
+		UE_LOG(LogCECloner, Warning, TEXT("%s : Invalid system for cloner layout"), *ClonerActor->GetActorNameOrLabel());
 		return;
 	}
 
@@ -2164,13 +1439,13 @@ void UCEClonerComponent::UpdateClonerMeshes()
 
 	if (!MeshRenderer)
 	{
-		UE_LOG(LogCEClonerComponent, Warning, TEXT("%s : Invalid mesh renderer for cloner system"), *ClonerActor->GetActorNameOrLabel());
+		UE_LOG(LogCECloner, Warning, TEXT("%s : Invalid mesh renderer for cloner system"), *ClonerActor->GetActorNameOrLabel());
 		return;
 	}
 
 	bool bMeshChanged = MeshRenderer->Meshes.Num() != ClonerTree.MergedBakedMeshes.Num();
 
-	if (bClonerMeshesDirty)
+	if (ClonerTree.bItemAttachmentsDirty)
 	{
 		// Resize mesh array properly
 		if (MeshRenderer->Meshes.Num() > ClonerTree.MergedBakedMeshes.Num())
@@ -2185,19 +1460,27 @@ void UCEClonerComponent::UpdateClonerMeshes()
 
 			FNiagaraMeshRendererMeshProperties& MeshProperties = !MeshRenderer->Meshes.IsValidIndex(Idx) ? MeshRenderer->Meshes.AddDefaulted_GetRef() : MeshRenderer->Meshes[Idx];
 			bMeshChanged |= MeshProperties.Mesh != StaticMesh;
-			MeshProperties.Mesh = StaticMesh && StaticMesh->GetNumTriangles(0) > 0 ? StaticMesh : nullptr;
+
+			if (StaticMesh && StaticMesh->GetNumTriangles(0) > 0)
+			{
+				MeshProperties.Mesh = StaticMesh;
+			}
+			else
+			{
+				MeshProperties.Mesh = nullptr;
+			}
 
 			if (ClonerTree.RootActors.IsValidIndex(Idx))
 			{
-				if (const FCEClonerAttachmentItem* Item = ClonerTree.ItemAttachmentMap.Find(ClonerTree.RootActors[Idx]))
+				if (const AActor* RootActor = ClonerTree.RootActors[Idx].Get())
 				{
-					MeshProperties.Rotation = GlobalRotation + Item->ActorTransform.Rotator();
-					MeshProperties.Scale = GlobalScale * Item->ActorTransform.GetScale3D();
+					MeshProperties.Rotation = GlobalRotation + RootActor->GetActorRotation();
+					MeshProperties.Scale = GlobalScale * RootActor->GetActorScale3D();
 				}
 			}
 		}
 
-		bClonerMeshesDirty = !ClonerTree.DirtyItemAttachments.IsEmpty();
+		ClonerTree.bItemAttachmentsDirty = !ClonerTree.DirtyItemAttachments.IsEmpty();
 	}
 
 	for (const TObjectPtr<UCEClonerExtensionBase>& ActiveExtension : ActiveExtensions)
@@ -2210,7 +1493,7 @@ void UCEClonerComponent::UpdateClonerMeshes()
 
 	if (bMeshChanged)
 	{
-		UE_LOG(LogCEClonerComponent, Log, TEXT("%s : Cloner mesh updated - %i baked meshes - %i rendered meshes"), *ClonerActor->GetActorNameOrLabel(), ClonerTree.MergedBakedMeshes.Num(), MeshRenderer->Meshes.Num())
+		UE_LOG(LogCECloner, Log, TEXT("%s : Cloner mesh updated - %i cached meshes - %i rendered meshes"), *ClonerActor->GetActorNameOrLabel(), ClonerTree.MergedBakedMeshes.Num(), MeshRenderer->Meshes.Num())
 	}
 
 	// Set new number of meshes in renderer
diff --git a/Engine/Plugins/Experimental/ClonerEffector/Source/ClonerEffector/Private/Cloner/Extensions/CEClonerEffectorExtension.cpp b/Engine/Plugins/Experimental/ClonerEffector/Source/ClonerEffector/Private/Cloner/Extensions/CEClonerEffectorExtension.cpp
--- a/Engine/Plugins/Experimental/ClonerEffector/Source/ClonerEffector/Private/Cloner/Extensions/CEClonerEffectorExtension.cpp
+++ b/Engine/Plugins/Experimental/ClonerEffector/Source/ClonerEffector/Private/Cloner/Extensions/CEClonerEffectorExtension.cpp
@@ -4,6 +4,7 @@
 
 #include "Cloner/CEClonerComponent.h"
 #include "Cloner/Layouts/CEClonerLayoutBase.h"
+#include "Cloner/Logs/CEClonerLogs.h"
 #include "Effector/CEEffectorActor.h"
 #include "Effector/CEEffectorComponent.h"
 #include "NiagaraDataInterfaceArrayInt.h"
@@ -14,8 +15,6 @@
 #include "Editor/EditorEngine.h"
 #endif
 
-DEFINE_LOG_CATEGORY_STATIC(LogCEClonerEffectorExtension, Log, All);
-
 #if WITH_EDITOR
 FName UCEClonerEffectorExtension::GetEffectorActorsWeakName()
 {
@@ -43,7 +42,7 @@ bool UCEClonerEffectorExtension::LinkEffector(AActor* InEffectorActor)
 
 	OnEffectorsChanged();
 
-	UE_LOG(LogCEClonerEffectorExtension, Log, TEXT("%s : Effector %s linked to Cloner"), *GetClonerComponent()->GetOwner()->GetActorNameOrLabel(), *InEffectorActor->GetActorNameOrLabel());
+	UE_LOG(LogCECloner, Log, TEXT("%s : Effector %s linked to Cloner"), *GetClonerComponent()->GetOwner()->GetActorNameOrLabel(), *InEffectorActor->GetActorNameOrLabel());
 
 	return true;
 }
@@ -59,7 +58,7 @@ bool UCEClonerEffectorExtension::UnlinkEffector(AActor* InEffectorActor)
 	{
 		OnEffectorsChanged();
 
-		UE_LOG(LogCEClonerEffectorExtension, Log, TEXT("%s : Effector %s unlinked from Cloner"), *GetClonerComponent()->GetOwner()->GetActorNameOrLabel(), *InEffectorActor->GetActorNameOrLabel());
+		UE_LOG(LogCECloner, Log, TEXT("%s : Effector %s unlinked from Cloner"), *GetClonerComponent()->GetOwner()->GetActorNameOrLabel(), *InEffectorActor->GetActorNameOrLabel());
 	}
 
 	return true;
diff --git a/Engine/Plugins/Experimental/ClonerEffector/Source/ClonerEffector/Private/Cloner/Extensions/CEClonerExtensionBase.cpp b/Engine/Plugins/Experimental/ClonerEffector/Source/ClonerEffector/Private/Cloner/Extensions/CEClonerExtensionBase.cpp
--- a/Engine/Plugins/Experimental/ClonerEffector/Source/ClonerEffector/Private/Cloner/Extensions/CEClonerExtensionBase.cpp
+++ b/Engine/Plugins/Experimental/ClonerEffector/Source/ClonerEffector/Private/Cloner/Extensions/CEClonerExtensionBase.cpp
@@ -3,8 +3,7 @@
 #include "Cloner/Extensions/CEClonerExtensionBase.h"
 
 #include "Cloner/CEClonerComponent.h"
-
-DEFINE_LOG_CATEGORY_STATIC(LogCEClonerExtensionBase, Log, All);
+#include "Cloner/Logs/CEClonerLogs.h"
 
 #if WITH_EDITOR
 FCEExtensionSection UCEClonerExtensionBase::GetExtensionSection() const
@@ -60,7 +59,7 @@ void UCEClonerExtensionBase::ActivateExtension()
 	{
 		bExtensionActive = true;
 		const UCEClonerComponent* ClonerComponent = GetClonerComponentChecked();
-		UE_LOG(LogCEClonerExtensionBase, Verbose, TEXT("%s : Cloner extension activated %s"), *ClonerComponent->GetOwner()->GetActorNameOrLabel(), *GetExtensionName().ToString());
+		UE_LOG(LogCECloner, Verbose, TEXT("%s : Cloner extension activated %s"), *ClonerComponent->GetOwner()->GetActorNameOrLabel(), *GetExtensionName().ToString());
 		OnExtensionActivated();
 	}
 }
@@ -71,7 +70,7 @@ void UCEClonerExtensionBase::DeactivateExtension()
 	{
 		bExtensionActive = false;
 		const UCEClonerComponent* ClonerComponent = GetClonerComponentChecked();
-		UE_LOG(LogCEClonerExtensionBase, Verbose, TEXT("%s : Cloner extension deactivated %s"), *ClonerComponent->GetOwner()->GetActorNameOrLabel(), *GetExtensionName().ToString());
+		UE_LOG(LogCECloner, Verbose, TEXT("%s : Cloner extension deactivated %s"), *ClonerComponent->GetOwner()->GetActorNameOrLabel(), *GetExtensionName().ToString());
 		OnExtensionDeactivated();
 	}
 }
diff --git a/Engine/Plugins/Experimental/ClonerEffector/Source/ClonerEffector/Private/Cloner/Extensions/CEClonerMeshRendererExtension.cpp b/Engine/Plugins/Experimental/ClonerEffector/Source/ClonerEffector/Private/Cloner/Extensions/CEClonerMeshRendererExtension.cpp
--- a/Engine/Plugins/Experimental/ClonerEffector/Source/ClonerEffector/Private/Cloner/Extensions/CEClonerMeshRendererExtension.cpp
+++ b/Engine/Plugins/Experimental/ClonerEffector/Source/ClonerEffector/Private/Cloner/Extensions/CEClonerMeshRendererExtension.cpp
@@ -3,13 +3,13 @@
 #include "Cloner/Extensions/CEClonerMeshRendererExtension.h"
 
 #include "Cloner/CEClonerComponent.h"
+#include "Cloner/Logs/CEClonerLogs.h"
 #include "Materials/Material.h"
 #include "Materials/MaterialInterface.h"
 #include "NiagaraSystem.h"
 #include "Settings/CEClonerEffectorSettings.h"
 #include "UObject/ConstructorHelpers.h"
-
-DEFINE_LOG_CATEGORY_STATIC(LogCEClonerMeshRenderer, Log, All);
+#include "Utilities/CEClonerEffectorUtilities.h"
 
 UCEClonerMeshRendererExtension::UCEClonerMeshRendererExtension()
 	: UCEClonerExtensionBase(
@@ -194,12 +194,14 @@ TArray<FNiagaraMeshMaterialOverride> UCEClonerMeshRendererExtension::GetOverride
 
 void UCEClonerMeshRendererExtension::OnOverrideMaterialOptionsChanged()
 {
-	if (IsValid(OverrideMaterial) && !UCEClonerComponent::IsMaterialUsageFlagSet(OverrideMaterial))
+	using namespace UE::ClonerEffector::Utilities;
+
+	if (IsValid(OverrideMaterial) && !IsMaterialUsageFlagSet(OverrideMaterial))
 	{
-		UE_LOG(LogCEClonerMeshRenderer, Warning, TEXT("%s : The override material (%s) you wish to use does not have the required usage flag (bUsedWithNiagaraMeshParticles) to work with the cloner, enable the flag on the material and save the asset"), *GetClonerComponent()->GetOwner()->GetActorNameOrLabel(), *OverrideMaterial->GetMaterial()->GetPathName());
+		UE_LOG(LogCECloner, Warning, TEXT("%s : The override material (%s) you wish to use does not have the required usage flag (bUsedWithNiagaraMeshParticles) to work with the cloner, enable the flag on the material and save the asset"), *GetClonerComponent()->GetOwner()->GetActorNameOrLabel(), *OverrideMaterial->GetMaterial()->GetPathName());
 
 #if WITH_EDITOR
-		UCEClonerComponent::ShowMaterialWarning(1);
+		ShowWarning(FText::Format(GetMaterialWarningText(), 1));
 #endif
 
 		OverrideMaterial = nullptr;
diff --git a/Engine/Plugins/Experimental/ClonerEffector/Source/ClonerEffector/Private/Cloner/Layouts/CEClonerLayoutBase.cpp b/Engine/Plugins/Experimental/ClonerEffector/Source/ClonerEffector/Private/Cloner/Layouts/CEClonerLayoutBase.cpp
--- a/Engine/Plugins/Experimental/ClonerEffector/Source/ClonerEffector/Private/Cloner/Layouts/CEClonerLayoutBase.cpp
+++ b/Engine/Plugins/Experimental/ClonerEffector/Source/ClonerEffector/Private/Cloner/Layouts/CEClonerLayoutBase.cpp
@@ -6,6 +6,7 @@
 #include "Cloner/CEClonerActor.h"
 #include "Cloner/CEClonerComponent.h"
 #include "Cloner/Extensions/CEClonerExtensionBase.h"
+#include "Cloner/Logs/CEClonerLogs.h"
 #include "Engine/Level.h"
 #include "Engine/World.h"
 #include "Misc/Guid.h"
@@ -17,8 +18,6 @@
 #include "UObject/Package.h"
 #include "UObject/UObjectGlobals.h"
 
-DEFINE_LOG_CATEGORY_STATIC(LogCEClonerLayoutBase, Log, All);
-
 bool UCEClonerLayoutBase::IsLayoutValid() const
 {
 	if (LayoutName.IsNone() || LayoutAssetPath.IsEmpty())
@@ -34,7 +33,7 @@ bool UCEClonerLayoutBase::IsLayoutValid() const
 
 	if (!TemplateNiagaraSystem || !BaseNiagaraSystem)
 	{
-		UE_LOG(LogCEClonerLayoutBase, Warning, TEXT("Cloner layout %s : Template system (%s) or base system (%s) is invalid"), *LayoutName.ToString(), *LayoutAssetPath, LayoutBaseAssetPath);
+		UE_LOG(LogCECloner, Warning, TEXT("Cloner layout %s : Template system (%s) or base system (%s) is invalid"), *LayoutName.ToString(), *LayoutAssetPath, LayoutBaseAssetPath);
 		return false;
 	}
 
@@ -53,7 +52,7 @@ bool UCEClonerLayoutBase::IsLayoutValid() const
 			if (!TemplateSystemParameters.Contains(SystemParameter))
 			{
 				bIsSystemBasedOnBaseAsset = false;
-				UE_LOG(LogCEClonerLayoutBase, Warning, TEXT("Cloner layout %s : Template system (%s) missing parameter (%s) from base system (%s)"), *LayoutName.ToString(), *LayoutAssetPath, *SystemParameter.ToString(), LayoutBaseAssetPath);
+				UE_LOG(LogCECloner, Warning, TEXT("Cloner layout %s : Template system (%s) missing parameter (%s) from base system (%s)"), *LayoutName.ToString(), *LayoutAssetPath, *SystemParameter.ToString(), LayoutBaseAssetPath);
 				break;
 			}
 		}
@@ -61,7 +60,7 @@ bool UCEClonerLayoutBase::IsLayoutValid() const
 
 	if (!bIsSystemBasedOnBaseAsset)
 	{
-		UE_LOG(LogCEClonerLayoutBase, Warning, TEXT("Cloner layout %s : Template system (%s) is not based off base system (%s)"), *LayoutName.ToString(), *LayoutAssetPath, LayoutBaseAssetPath);
+		UE_LOG(LogCECloner, Warning, TEXT("Cloner layout %s : Template system (%s) is not based off base system (%s)"), *LayoutName.ToString(), *LayoutAssetPath, LayoutBaseAssetPath);
 	}
 #if WITH_EDITOR
 	else
@@ -70,11 +69,11 @@ bool UCEClonerLayoutBase::IsLayoutValid() const
 
 		if (LayoutSystemHash.IsEmpty())
 		{
-			UE_LOG(LogCEClonerLayoutBase, Warning, TEXT("Cloner layout %s : Template system (%s) hash could not be calculated"), *LayoutName.ToString(), *LayoutAssetPath);
+			UE_LOG(LogCECloner, Warning, TEXT("Cloner layout %s : Template system (%s) hash could not be calculated"), *LayoutName.ToString(), *LayoutAssetPath);
 			return false;
 		}
 
-		UE_LOG(LogCEClonerLayoutBase, Verbose, TEXT("Cloner layout %s : Template system (%s) hash is %s"), *LayoutName.ToString(), *LayoutAssetPath, *LayoutSystemHash);
+		UE_LOG(LogCECloner, Verbose, TEXT("Cloner layout %s : Template system (%s) hash is %s"), *LayoutName.ToString(), *LayoutAssetPath, *LayoutSystemHash);
 	}
 #endif
 
@@ -104,7 +103,7 @@ void UCEClonerLayoutBase::LoadLayout()
 	{
 		if (IsSystemHashMatching())
 		{
-			UE_LOG(LogCEClonerLayoutBase, Verbose, TEXT("%s : Cloner layout %s using cached system %s"), *GetClonerActor()->GetActorNameOrLabel(), *LayoutName.ToString(), *CachedSystemHash)
+			UE_LOG(LogCECloner, Verbose, TEXT("%s : Cloner layout %s using cached system %s"), *GetClonerActor()->GetActorNameOrLabel(), *LayoutName.ToString(), *CachedSystemHash)
 
 			CacheMeshRenderer();
 			OnSystemLoaded();
@@ -116,7 +115,7 @@ void UCEClonerLayoutBase::LoadLayout()
 		LayoutHash = GetLayoutHash();
 #endif
 
-		UE_LOG(LogCEClonerLayoutBase, Warning, TEXT("%s : Cloner layout %s skipping cached system %s due to hash mismatch %s"), *GetClonerActor()->GetActorNameOrLabel(), *LayoutName.ToString(), *CachedSystemHash, *LayoutHash)
+		UE_LOG(LogCECloner, Warning, TEXT("%s : Cloner layout %s skipping cached system %s due to hash mismatch %s"), *GetClonerActor()->GetActorNameOrLabel(), *LayoutName.ToString(), *CachedSystemHash, *LayoutHash)
 
 		NiagaraSystem->MarkAsGarbage();
 		NiagaraSystem = nullptr;
@@ -162,7 +161,7 @@ void UCEClonerLayoutBase::LoadLayout()
 	FPackagePath::TryFromPackageName(TEXT("/Game/Temp/") + GetLayoutName().ToString() + TEXT("_") + FGuid::NewGuid().ToString(), CustomPackagePath);
 	CustomPackagePath.SetHeaderExtension(EPackageExtension::Asset);
 
-	UE_LOG(LogCEClonerLayoutBase, Verbose, TEXT("%s : Cloner layout load requested %s - Template system %s - Package %s"), *GetClonerActor()->GetActorNameOrLabel(), *LayoutName.ToString(), *LayoutAssetPath, *CustomPackagePath.GetPackageFName().ToString())
+	UE_LOG(LogCECloner, Verbose, TEXT("%s : Cloner layout load requested %s - Template system %s - Package %s"), *GetClonerActor()->GetActorNameOrLabel(), *LayoutName.ToString(), *LayoutAssetPath, *CustomPackagePath.GetPackageFName().ToString())
 
 	FLoadPackageAsyncOptionalParams Params;
 	Params.PackagePriority = INT32_MAX;
@@ -198,7 +197,7 @@ bool UCEClonerLayoutBase::UnloadLayout()
 
 	MeshRenderer = nullptr;
 
-	UE_LOG(LogCEClonerLayoutBase, Verbose, TEXT("%s : Cloner layout unloaded %s"), *GetClonerActor()->GetActorNameOrLabel(), *LayoutName.ToString())
+	UE_LOG(LogCECloner, Verbose, TEXT("%s : Cloner layout unloaded %s"), *GetClonerActor()->GetActorNameOrLabel(), *LayoutName.ToString())
 
 	OnLayoutUnloaded();
 
@@ -239,7 +238,7 @@ bool UCEClonerLayoutBase::ActivateLayout()
 
 	ClonerComponent->SetAsset(NiagaraSystem);
 
-	UE_LOG(LogCEClonerLayoutBase, Verbose, TEXT("%s : Cloner layout activated %s"), *GetClonerActor()->GetActorNameOrLabel(), *LayoutName.ToString())
+	UE_LOG(LogCECloner, Verbose, TEXT("%s : Cloner layout activated %s"), *GetClonerActor()->GetActorNameOrLabel(), *LayoutName.ToString())
 
 	OnLayoutActive();
 
@@ -267,7 +266,7 @@ bool UCEClonerLayoutBase::DeactivateLayout()
 	NiagaraSystem->KillAllActiveCompilations();
 #endif
 
-	UE_LOG(LogCEClonerLayoutBase, Verbose, TEXT("%s : Cloner layout deactivated %s"), *GetClonerActor()->GetActorNameOrLabel(), *LayoutName.ToString())
+	UE_LOG(LogCECloner, Verbose, TEXT("%s : Cloner layout deactivated %s"), *GetClonerActor()->GetActorNameOrLabel(), *LayoutName.ToString())
 
 	OnLayoutInactive();
 
@@ -382,13 +381,13 @@ void UCEClonerLayoutBase::OnSystemLoaded()
 
 	if (bLayoutLoaded)
 	{
-		UE_LOG(LogCEClonerLayoutBase, Verbose, TEXT("%s : Cloner layout loaded %s - Template system %s"), *GetClonerActor()->GetActorNameOrLabel(), *LayoutName.ToString(), *LayoutAssetPath)
+		UE_LOG(LogCECloner, Verbose, TEXT("%s : Cloner layout loaded %s - Template system %s"), *GetClonerActor()->GetActorNameOrLabel(), *LayoutName.ToString(), *LayoutAssetPath)
 
 		OnLayoutLoaded();
 	}
 	else
 	{
-		UE_LOG(LogCEClonerLayoutBase, Warning, TEXT("%s : Cloner layout load failed %s - Template system %s"), *GetClonerActor()->GetActorNameOrLabel(), *LayoutName.ToString(), *LayoutAssetPath)
+		UE_LOG(LogCECloner, Warning, TEXT("%s : Cloner layout load failed %s - Template system %s"), *GetClonerActor()->GetActorNameOrLabel(), *LayoutName.ToString(), *LayoutAssetPath)
 	}
 
 	OnClonerLayoutLoadedDelegate.Broadcast(this, bLayoutLoaded);
@@ -465,7 +464,7 @@ void UCEClonerLayoutBase::OnLevelCleanup()
 {
 	if (IsLayoutLoaded())
 	{
-		UE_LOG(LogCEClonerLayoutBase, Log, TEXT("%s : Cloner layout cleanup %s"), *GetClonerActor()->GetActorNameOrLabel(), *LayoutName.ToString())
+		UE_LOG(LogCECloner, Log, TEXT("%s : Cloner layout cleanup %s"), *GetClonerActor()->GetActorNameOrLabel(), *LayoutName.ToString())
 
 		DeactivateLayout();
 
@@ -484,7 +483,7 @@ void UCEClonerLayoutBase::CleanOwnedSystem() const
 	{
 		if (OwnedObject && OwnedObject->IsA<UNiagaraSystem>())
 		{
-			UE_LOG(LogCEClonerLayoutBase, Warning, TEXT("%s : Cloner layout %s cleaning owned system %s"), *GetClonerActor()->GetActorNameOrLabel(), *LayoutName.ToString(), *OwnedObject->GetName())
+			UE_LOG(LogCECloner, Warning, TEXT("%s : Cloner layout %s cleaning owned system %s"), *GetClonerActor()->GetActorNameOrLabel(), *LayoutName.ToString(), *OwnedObject->GetName())
 			OwnedObject->MarkAsGarbage();
 		}
 	}
diff --git a/Engine/Plugins/Experimental/ClonerEffector/Source/ClonerEffector/Private/Cloner/Logs/CEClonerLogs.h b/Engine/Plugins/Experimental/ClonerEffector/Source/ClonerEffector/Private/Cloner/Logs/CEClonerLogs.h
new file mode 100644
--- /dev/null
+++ b/Engine/Plugins/Experimental/ClonerEffector/Source/ClonerEffector/Private/Cloner/Logs/CEClonerLogs.h
@@ -0,0 +1,7 @@
+// Copyright Epic Games, Inc. All Rights Reserved.
+
+#pragma once
+
+#include "Logging/LogMacros.h"
+
+DECLARE_LOG_CATEGORY_EXTERN(LogCECloner, Log, All);
\ No newline at end of file
diff --git a/Engine/Plugins/Experimental/ClonerEffector/Source/ClonerEffector/Private/Subsystems/CEClonerSubsystem.cpp b/Engine/Plugins/Experimental/ClonerEffector/Source/ClonerEffector/Private/Subsystems/CEClonerSubsystem.cpp
--- a/Engine/Plugins/Experimental/ClonerEffector/Source/ClonerEffector/Private/Subsystems/CEClonerSubsystem.cpp
+++ b/Engine/Plugins/Experimental/ClonerEffector/Source/ClonerEffector/Private/Subsystems/CEClonerSubsystem.cpp
@@ -4,6 +4,8 @@
 
 #include "Cloner/CEClonerActor.h"
 #include "Cloner/CEClonerComponent.h"
+#include "Cloner/Attachments/CEClonerAttachmentTreeBehavior.h"
+#include "Cloner/Attachments/CEClonerSceneTreeCustomResolver.h"
 #include "Cloner/Extensions/CEClonerEffectorExtension.h"
 #include "Cloner/Extensions/CEClonerExtensionBase.h"
 #include "Cloner/Layouts/CEClonerCircleLayout.h"
@@ -16,9 +18,13 @@
 #include "Cloner/Layouts/CEClonerSphereRandomLayout.h"
 #include "Cloner/Layouts/CEClonerSphereUniformLayout.h"
 #include "Cloner/Layouts/CEClonerSplineLayout.h"
+#include "Cloner/Logs/CEClonerLogs.h"
 #include "Engine/Engine.h"
+#include "GameFramework/Actor.h"
+#include "Materials/Material.h"
 #include "UObject/Class.h"
 #include "UObject/UObjectIterator.h"
+#include "Utilities/CEClonerEffectorUtilities.h"
 
 #if WITH_EDITOR
 #include "Editor.h"
@@ -27,6 +33,7 @@
 
 UCEClonerSubsystem::FOnSubsystemInitialized UCEClonerSubsystem::OnSubsystemInitializedDelegate;
 UCEClonerSubsystem::FOnClonerSetEnabled UCEClonerSubsystem::OnClonerSetEnabledDelegate;
+UCEClonerSubsystem::FOnGetSceneTreeResolver UCEClonerSubsystem::OnGetSceneTreeResolverDelegate;
 
 #define LOCTEXT_NAMESPACE "CEEffectorSubsystem"
 
@@ -58,7 +65,31 @@ void UCEClonerSubsystem::Initialize(FSubsystemCollectionBase& Collection)
 	// Scan for new layouts
 	ScanForRegistrableClasses();
 
+	// Tree implementation
+	{
+		TFunction<TSharedRef<ICEClonerAttachmentTreeBehavior>()> GroupImplCreator = []()
+		{
+			return MakeShared<FCEClonerAttachmentGroupBehavior>();
+		};
+		RegisterAttachmentTreeBehavior(TEXT("Group"), GroupImplCreator);
+		
+		TFunction<TSharedRef<ICEClonerAttachmentTreeBehavior>()> FlatImplCreator = []()
+		{
+			return MakeShared<FCEClonerAttachmentFlatBehavior>();
+		};
+		RegisterAttachmentTreeBehavior(TEXT("Flat"), FlatImplCreator);
+	}
+
 	OnSubsystemInitializedDelegate.Broadcast();
+
+	FWorldDelegates::OnWorldCleanup.AddUObject(this, &UCEClonerSubsystem::OnWorldCleanup);
+}
+
+void UCEClonerSubsystem::Deinitialize()
+{
+	Super::Deinitialize();
+
+	FWorldDelegates::OnWorldCleanup.RemoveAll(this);
 }
 
 bool UCEClonerSubsystem::RegisterLayoutClass(UClass* InClonerLayoutClass)
@@ -137,21 +168,6 @@ bool UCEClonerSubsystem::IsLayoutClassRegistered(UClass* InClonerLayoutClass)
 	return false;
 }
 
-void UCEClonerSubsystem::RegisterCustomActorResolver(FOnGetOrderedActors InCustomResolver)
-{
-	ActorResolver = InCustomResolver;
-}
-
-void UCEClonerSubsystem::UnregisterCustomActorResolver()
-{
-	ActorResolver = FOnGetOrderedActors();
-}
-
-UCEClonerSubsystem::FOnGetOrderedActors& UCEClonerSubsystem::GetCustomActorResolver()
-{
-	return ActorResolver;
-}
-
 bool UCEClonerSubsystem::RegisterExtensionClass(UClass* InClass)
 {
 	if (!IsValid(InClass))
@@ -437,6 +453,102 @@ AActor* UCEClonerSubsystem::CreateClonerWithActors(UWorld* InWorld, const TSet<A
 	return NewClonerActor;
 }
 
+void UCEClonerSubsystem::FireMaterialWarning(const AActor* InClonerActor, const AActor* InContextActor, const TArray<TWeakObjectPtr<UMaterialInterface>>& InUnsetMaterials)
+{
+	if (!IsValid(InContextActor) || !IsValid(InClonerActor) || InUnsetMaterials.IsEmpty())
+	{
+		return;
+	}
+
+	UE_LOG(LogCECloner, Warning, TEXT("%s : %i unsupported material(s) detected due to missing niagara usage flag (bUsedWithNiagaraMeshParticles) on actor (%s), see logs below"), *InClonerActor->GetActorNameOrLabel(), InUnsetMaterials.Num(), *InContextActor->GetActorNameOrLabel());
+
+	for (const TWeakObjectPtr<UMaterialInterface>& UnsetMaterialWeak : InUnsetMaterials)
+	{
+		if (UMaterialInterface* UnsetMaterial = UnsetMaterialWeak.Get())
+		{
+			UE_LOG(LogCECloner, Warning, TEXT("%s : The following materials (%s) on actor (%s) does not have the usage flag (bUsedWithNiagaraMeshParticles) set to work with the cloner, set the flag and resave the asset to avoid this warning"), *InClonerActor->GetActorNameOrLabel(), *UnsetMaterial->GetMaterial()->GetPathName(), *InContextActor->GetActorNameOrLabel());
+		}
+	}
+
+#if WITH_EDITOR
+	using namespace UE::ClonerEffector::Utilities;
+
+	// Fire warning notification when invalid materials are found and at least 5s has elapsed since last notification
+	constexpr double MinNotificationElapsedTime = 5.0;
+	const double CurrentTime = FApp::GetCurrentTime();
+
+	if (CurrentTime - LastNotificationTime > MinNotificationElapsedTime)
+	{
+		LastNotificationTime = CurrentTime;
+		ShowWarning(FText::Format(GetMaterialWarningText(), InUnsetMaterials.Num()));
+	}
+#endif
+}
+
+bool UCEClonerSubsystem::RegisterAttachmentTreeBehavior(FName InName, TFunction<TSharedRef<ICEClonerAttachmentTreeBehavior>()>& InCreator)
+{
+	if (!TreeBehaviorCreators.Contains(InName))
+	{
+		TreeBehaviorCreators.Add(InName, InCreator);
+		return true;
+	}
+
+	return false;
+}
+
+bool UCEClonerSubsystem::UnregisterAttachmentTreeBehavior(FName InName)
+{
+	return TreeBehaviorCreators.Remove(InName) > 0;
+}
+
+TArray<FName> UCEClonerSubsystem::GetAttachmentTreeBehaviorNames() const
+{
+	TArray<FName> TreeImplNames;
+	TreeBehaviorCreators.GenerateKeyArray(TreeImplNames);
+	return TreeImplNames;
+}
+
+TSharedPtr<ICEClonerAttachmentTreeBehavior> UCEClonerSubsystem::CreateAttachmentTreeBehavior(FName InName) const
+{
+	TSharedPtr<ICEClonerAttachmentTreeBehavior> TreeBehavior = nullptr;
+
+	if (const TFunction<TSharedRef<ICEClonerAttachmentTreeBehavior>()>* BehaviorCreator = TreeBehaviorCreators.Find(InName))
+	{
+		TreeBehavior = (*BehaviorCreator)();
+	}
+
+	return TreeBehavior;
+}
+
+TSharedPtr<ICEClonerSceneTreeCustomResolver> UCEClonerSubsystem::FindCustomLevelSceneTreeResolver(ULevel* InLevel)
+{
+	TSharedPtr<ICEClonerSceneTreeCustomResolver> Resolver = nullptr;
+
+	if (!InLevel)
+	{
+		return Resolver;
+	}
+
+	if (const TSharedRef<ICEClonerSceneTreeCustomResolver>* CachedResolver = LevelCustomResolvers.Find(InLevel))
+	{
+		Resolver = *CachedResolver;
+	}
+
+	if (!Resolver.IsValid() && OnGetSceneTreeResolverDelegate.IsBound())
+	{
+		Resolver = OnGetSceneTreeResolverDelegate.Execute(InLevel);
+
+		if (Resolver.IsValid())
+		{
+			LevelCustomResolvers.Add(InLevel, Resolver.ToSharedRef());
+			InLevel->OnCleanupLevel.AddUObject(this, &UCEClonerSubsystem::OnLevelCleanup, InLevel);
+			Resolver->Activate();
+		}
+	}
+
+	return Resolver;
+}
+
 TSet<FName> UCEClonerSubsystem::GetLayoutNames() const
 {
 	TArray<FName> LayoutNames;
@@ -511,4 +623,25 @@ void UCEClonerSubsystem::ScanForRegistrableClasses()
 	}
 }
 
-#undef LOCTEXT_NAMESPACE
\ No newline at end of file
+void UCEClonerSubsystem::OnLevelCleanup(ULevel* InLevel)
+{
+	if (const TSharedRef<ICEClonerSceneTreeCustomResolver>* CustomResolver = LevelCustomResolvers.Find(InLevel))
+	{
+		InLevel->OnCleanupLevel.RemoveAll(this);
+		(*CustomResolver)->Deactivate();
+		LevelCustomResolvers.Remove(InLevel);
+	}
+}
+
+void UCEClonerSubsystem::OnWorldCleanup(UWorld* InWorld, bool, bool bInCleanupResources)
+{
+	if (bInCleanupResources)
+	{
+		for (ULevel* Level : InWorld->GetLevels())
+		{
+			OnLevelCleanup(Level);
+		}
+	}
+}
+
+#undef LOCTEXT_NAMESPACE
diff --git a/Engine/Plugins/Experimental/ClonerEffector/Source/ClonerEffector/Private/Utilities/CEClonerEffectorUtilities.cpp b/Engine/Plugins/Experimental/ClonerEffector/Source/ClonerEffector/Private/Utilities/CEClonerEffectorUtilities.cpp
new file mode 100644
--- /dev/null
+++ b/Engine/Plugins/Experimental/ClonerEffector/Source/ClonerEffector/Private/Utilities/CEClonerEffectorUtilities.cpp
@@ -0,0 +1,142 @@
+// Copyright Epic Games, Inc. All Rights Reserved.
+
+#include "Utilities/CEClonerEffectorUtilities.h"
+
+#include "Cloner/CEClonerComponent.h"
+#include "Materials/Material.h"
+#include "UObject/Package.h"
+
+#if WITH_EDITOR
+#include "Framework/Notifications/NotificationManager.h"
+#include "Widgets/Notifications/SNotificationList.h"
+#endif
+
+#define LOCTEXT_NAMESPACE "CEClonerEffectorUtilities"
+
+#if WITH_EDITOR
+const FText& UE::ClonerEffector::Utilities::GetMaterialWarningText()
+{
+	static const FText MaterialWarningText = LOCTEXT("MaterialsMissingUsageFlag", "Detected {0} material(s) with missing usage flag required to work properly with cloner (See logs)");
+	return MaterialWarningText;
+}
+
+void UE::ClonerEffector::Utilities::ShowWarning(const FText& InWarning)
+{
+	if (!InWarning.IsEmpty())
+	{
+		FNotificationInfo NotificationInfo(InWarning);
+		NotificationInfo.ExpireDuration = 10.f;
+		NotificationInfo.bFireAndForget = true;
+		NotificationInfo.Image = FAppStyle::GetBrush("Icons.WarningWithColor");
+
+		FSlateNotificationManager::Get().AddNotification(NotificationInfo);
+	}
+}
+#endif
+
+bool UE::ClonerEffector::Utilities::IsMaterialDirtyable(const UMaterialInterface* InMaterial)
+{
+	const UMaterial* BaseMaterial = InMaterial->GetMaterial_Concurrent();
+	const FString ContentFolder = FPaths::ConvertRelativePathToFull(FPaths::ProjectContentDir());
+
+	const UPackage* MaterialPackage = BaseMaterial->GetPackage();
+	const FPackagePath& LoadedPath = MaterialPackage->GetLoadedPath();
+	const FString PackagePath = FPaths::ConvertRelativePathToFull(LoadedPath.GetLocalFullPath());
+	const FString MaterialPath = BaseMaterial->GetPathName();
+
+	const bool bTransientPackage = MaterialPackage == GetTransientPackage() || MaterialPath.StartsWith("/Temp/");
+	const bool bContentFolder = PackagePath.StartsWith(ContentFolder);
+
+	return bTransientPackage || bContentFolder;
+}
+
+bool UE::ClonerEffector::Utilities::IsMaterialUsageFlagSet(const UMaterialInterface* InMaterial)
+{
+	if (InMaterial)
+	{
+		if (const UMaterial* Material = InMaterial->GetMaterial_Concurrent())
+		{
+			return Material->GetUsageByFlag(EMaterialUsage::MATUSAGE_NiagaraMeshParticles);
+		}
+	}
+
+	return false;
+}
+
+bool UE::ClonerEffector::Utilities::FilterSupportedMaterials(TArray<TWeakObjectPtr<UMaterialInterface>>& InMaterials, TArray<TWeakObjectPtr<UMaterialInterface>>& OutUnsetMaterials, UMaterialInterface* InDefaultMaterial)
+{
+	check(InDefaultMaterial)
+
+	OutUnsetMaterials.Reset(InMaterials.Num());
+
+	for (int32 Index = 0; Index < InMaterials.Num(); Index++)
+	{
+		UMaterialInterface* PreviousMaterialInterface = InMaterials[Index].Get();
+
+		UMaterialInterface* NewMaterialInterface = PreviousMaterialInterface;
+
+		if (FilterSupportedMaterial(NewMaterialInterface, InDefaultMaterial))
+		{
+			// Add original material to unset list
+			OutUnsetMaterials.Add(PreviousMaterialInterface);
+		}
+
+		// Replace material
+		InMaterials[Index] = NewMaterialInterface;
+	}
+
+	return OutUnsetMaterials.IsEmpty();
+}
+
+bool UE::ClonerEffector::Utilities::FilterSupportedMaterial(UMaterialInterface*& InMaterial, UMaterialInterface* InDefaultMaterial)
+{
+	if (InMaterial && !IsMaterialUsageFlagSet(InMaterial))
+	{
+		// Replace material if dirtyable and not in read only location
+		if (!IsMaterialDirtyable(InMaterial))
+		{
+			InMaterial = InDefaultMaterial;
+		}
+
+		return true;
+	}
+
+	return false;
+}
+
+void UE::ClonerEffector::Utilities::SetActorVisibility(AActor* InActor, bool bInVisibility, ECEClonerActorVisibility InTarget)
+{
+	if (!InActor)
+	{
+		return;
+	}
+
+#if WITH_EDITOR
+	if (EnumHasAnyFlags(InTarget, ECEClonerActorVisibility::Editor))
+	{
+		InActor->SetIsTemporarilyHiddenInEditor(!bInVisibility);
+	}
+#endif
+
+	if (EnumHasAnyFlags(InTarget, ECEClonerActorVisibility::Game))
+	{
+		InActor->SetActorHiddenInGame(!bInVisibility);
+	}
+}
+
+AActor* UE::ClonerEffector::Utilities::FindClonerActor(AActor* InActor)
+{
+	if (!InActor)
+	{
+		return nullptr;
+	}
+
+	if (!!InActor->FindComponentByClass<UCEClonerComponent>())
+	{
+		return InActor;
+	}
+
+	return FindClonerActor(InActor->GetAttachParentActor());
+}
+
+#undef LOCTEXT_NAMESPACE
diff --git a/Engine/Plugins/Experimental/ClonerEffector/Source/ClonerEffector/Public/CEClonerEffectorShared.h b/Engine/Plugins/Experimental/ClonerEffector/Source/ClonerEffector/Public/CEClonerEffectorShared.h
--- a/Engine/Plugins/Experimental/ClonerEffector/Source/ClonerEffector/Public/CEClonerEffectorShared.h
+++ b/Engine/Plugins/Experimental/ClonerEffector/Source/ClonerEffector/Public/CEClonerEffectorShared.h
@@ -196,19 +196,6 @@ enum class ECEClonerEffectorPushDirection : uint8
 	Random
 };
 
-UENUM()
-enum class ECEClonerAttachmentStatus : uint8
-{
-	/** Item should be removed, no longer valid */
-	Invalid,
-	/** Item should be updated, changes detected */
-	Outdated,
-	/** Item is up to date, no changes needed */
-	Updated,
-	/** Item is being updated at the moment */
-	Updating
-};
-
 /** Enumerate all texture channels to sample for constraint */
 UENUM()
 enum class ECEClonerTextureSampleChannel : uint8
@@ -300,115 +287,6 @@ enum class ECEClonerActorVisibility : uint8
 	All = Game | Editor
 };
 
-USTRUCT()
-struct FCEClonerAttachmentItem
-{
-	GENERATED_BODY()
-
-	/** Children attached to this actor, order is not important here as they are combined into one */
-	UPROPERTY()
-	TSet<TWeakObjectPtr<AActor>> ChildrenActors;
-
-	/** Current cloner attached actor represented by this item */
-	UPROPERTY()
-	TWeakObjectPtr<AActor> ItemActor;
-
-	/** Parent of this item actor or null if below cloner */
-	UPROPERTY()
-	TWeakObjectPtr<AActor> ParentActor;
-
-	/** Actual baked mesh for the current actor item,
-		dynamic because it's easier for merging them together and avoid conversions again */
-	UPROPERTY(Transient)
-	TObjectPtr<UDynamicMesh> BakedMesh;
-
-	/** Actual baked materials for the current actor item */
-	UPROPERTY(Transient)
-	TArray<TWeakObjectPtr<UMaterialInterface>> BakedMaterials;
-
-	/** Status of the baked mesh, does it needs to be updated */
-	UPROPERTY()
-	ECEClonerAttachmentStatus MeshStatus = ECEClonerAttachmentStatus::Outdated;
-
-	/** Last actor item transform, used to trigger an update if changed */
-	UPROPERTY()
-	FTransform ActorTransform;
-
-	/** Item is root cloner actor */
-	UPROPERTY()
-	bool bRootItem = false;
-
-	/** Item setup is done (visibility, delegates) once attached/loaded */
-	UPROPERTY(Transient)
-	bool bSetupDone = false;
-
-	/** Status of this attachment item */
-	UPROPERTY(Transient)
-	ECEClonerAttachmentStatus Status = ECEClonerAttachmentStatus::Outdated;
-
-#if WITH_EDITORONLY_DATA
-	/** Cached pivot to detect changes in geometry */
-	UPROPERTY(Transient)
-	FVector Origin;
-
-	/** Cached extent to detect changes in geometry */
-	UPROPERTY(Transient)
-	FVector Extent;
-#endif
-
-	friend uint32 GetTypeHash(const FCEClonerAttachmentItem& InItem)
-	{
-		return GetTypeHash(InItem.ItemActor);
-	}
-
-	bool operator==(const FCEClonerAttachmentItem& InOther) const
-	{
-		return ItemActor.Get() == InOther.ItemActor.Get();
-	}
-
-#if WITH_EDITOR
-	bool CheckBoundsChanged(bool bInUpdate);
-#endif
-};
-
-USTRUCT()
-struct FCEClonerAttachmentTree
-{
-	GENERATED_BODY()
-
-	/** All cloner attached actor items */
-	UPROPERTY()
-	TMap<TWeakObjectPtr<AActor>, FCEClonerAttachmentItem> ItemAttachmentMap;
-
-	/** Actors directly attached to the cloner actor, order is important here */
-	UPROPERTY()
-	TArray<TWeakObjectPtr<AActor>> RootActors;
-
-	/**
-	 * Merged static meshes corresponding to the root actors for niagara,
-	 * result of merging dynamic meshes together
-	 */
-	UPROPERTY()
-	TArray<TObjectPtr<UStaticMesh>> MergedBakedMeshes;
-
-	/** Status of the cloner tree */
-	UPROPERTY(Transient)
-	ECEClonerAttachmentStatus Status = ECEClonerAttachmentStatus::Updated;
-
-	/** Attachment items that are dirty and need an update */
-	UPROPERTY(Transient)
-	TSet<TWeakObjectPtr<AActor>> DirtyItemAttachments;
-
-	void Reset()
-	{
-		ItemAttachmentMap.Reset();
-		RootActors.Reset();
-		MergedBakedMeshes.Reset();
-		DirtyItemAttachments.Reset();
-		Status = ECEClonerAttachmentStatus::Updated;
-	}
-};
-
 USTRUCT()
 struct FCEClonerEffectorChannelData
 {
diff --git a/Engine/Plugins/Experimental/ClonerEffector/Source/ClonerEffector/Public/Cloner/Attachments/CEClonerAttachmentItem.h b/Engine/Plugins/Experimental/ClonerEffector/Source/ClonerEffector/Public/Cloner/Attachments/CEClonerAttachmentItem.h
new file mode 100644
--- /dev/null
+++ b/Engine/Plugins/Experimental/ClonerEffector/Source/ClonerEffector/Public/Cloner/Attachments/CEClonerAttachmentItem.h
@@ -0,0 +1,93 @@
+// Copyright Epic Games, Inc. All Rights Reserved.
+
+#pragma once
+
+#include "Containers/Array.h"
+#include "CEClonerAttachmentTypes.h"
+#include "UObject/WeakObjectPtr.h"
+#include "CEClonerAttachmentItem.generated.h"
+
+class AActor;
+class UDynamicMesh;
+class UMaterialInterface;
+
+USTRUCT()
+struct FCEClonerAttachmentItem
+{
+	GENERATED_BODY()
+
+	/** Children attached to this actor, order is not important here as they are combined into one */
+	UPROPERTY()
+	TSet<TWeakObjectPtr<AActor>> ChildrenActors;
+
+	/** Current cloner attached actor represented by this item */
+	UPROPERTY()
+	TWeakObjectPtr<AActor> ItemActor;
+
+	/** Parent of this item actor or null if below cloner */
+	UPROPERTY()
+	TWeakObjectPtr<AActor> ParentActor;
+
+	/** Actual baked mesh for the current actor item,
+		dynamic because it's easier for merging them together and avoid conversions again */
+	UPROPERTY(Transient)
+	TObjectPtr<UDynamicMesh> BakedMesh;
+
+	/** Actual baked materials for the current actor item */
+	UPROPERTY(Transient)
+	TArray<TWeakObjectPtr<UMaterialInterface>> BakedMaterials;
+
+	/** Status of the baked mesh, does it needs to be updated */
+	UPROPERTY()
+	ECEClonerAttachmentStatus MeshStatus = ECEClonerAttachmentStatus::Outdated;
+
+	/** Last actor item transform, used to trigger an update if changed */
+	UPROPERTY()
+	FTransform ActorTransform = FTransform::Identity;
+
+	/** Item is root cloner actor */
+	UPROPERTY()
+	bool bRootItem = false;
+
+	/** Item setup is done (visibility, delegates) once attached/loaded */
+	UPROPERTY(Transient)
+	bool bSetupDone = false;
+
+	/** Status of this attachment item */
+	UPROPERTY(Transient)
+	ECEClonerAttachmentStatus Status = ECEClonerAttachmentStatus::Outdated;
+
+#if WITH_EDITORONLY_DATA
+	/** Cached pivot to detect changes in geometry */
+	UPROPERTY(Transient)
+	FVector Origin = FVector::ZeroVector;
+
+	/** Cached extent to detect changes in geometry */
+	UPROPERTY(Transient)
+	FVector Extent = FVector::ZeroVector;
+#endif
+
+	friend uint32 GetTypeHash(const FCEClonerAttachmentItem& InItem)
+	{
+		return GetTypeHash(InItem.ItemActor);
+	}
+
+	bool operator==(const FCEClonerAttachmentItem& InOther) const
+	{
+		return ItemActor.Get() == InOther.ItemActor.Get();
+	}
+
+#if WITH_EDITOR
+	/** Check this item bounds without recursing */
+	bool CheckBoundsChanged(bool bInUpdate);
+
+	/** Calculate baked mesh bounds */
+	FBox GetBakedMeshBounds() const;
+
+	/** Get local attachment bounds */
+	FBox GetAttachmentBounds() const;
+
+	/** Check this item materials to ensure they are supported */
+	bool CheckMaterialsChanged(bool bInUpdate, TArray<TWeakObjectPtr<UMaterialInterface>>* OutInvalidMaterials = nullptr);
+#endif
+};
\ No newline at end of file
diff --git a/Engine/Plugins/Experimental/ClonerEffector/Source/ClonerEffector/Public/Cloner/Attachments/CEClonerAttachmentTree.h b/Engine/Plugins/Experimental/ClonerEffector/Source/ClonerEffector/Public/Cloner/Attachments/CEClonerAttachmentTree.h
new file mode 100644
--- /dev/null
+++ b/Engine/Plugins/Experimental/ClonerEffector/Source/ClonerEffector/Public/Cloner/Attachments/CEClonerAttachmentTree.h
@@ -0,0 +1,133 @@
+// Copyright Epic Games, Inc. All Rights Reserved.
+
+#pragma once
+
+#include "CEClonerAttachmentItem.h"
+#include "CEClonerAttachmentTypes.h"
+#include "Components/SceneComponent.h"
+#include "Containers/Array.h"
+#include "Containers/Map.h"
+#include "Containers/Set.h"
+#include "Templates/SharedPointerFwd.h"
+#include "UObject/ObjectPtr.h"
+#include "CEClonerAttachmentTree.generated.h"
+
+class AActor;
+class ICEClonerAttachmentTreeBehavior;
+class UObject;
+class UStaticMesh;
+
+/**
+ * Represents a logical tree used by the cloner to render its attachments
+ * Could differ from the actual physical tree used in the scene
+ */
+USTRUCT()
+struct FCEClonerAttachmentTree
+{
+	GENERATED_BODY()
+
+	DECLARE_DELEGATE_TwoParams(FOnAttachmentChanged, AActor*, FCEClonerAttachmentItem&)
+
+	FOnAttachmentChanged::RegistrationType& OnItemAttached()
+	{
+		return OnItemAttachedDelegate;
+	}
+
+	FOnAttachmentChanged::RegistrationType& OnItemDetached()
+	{
+		return OnItemDetachedDelegate;
+	}
+
+	/** Sets tree attachment root to scan from */
+	void SetAttachmentRoot(AActor* InActor);
+
+	/** Gets tree attachment root */
+	AActor* GetAttachmentRoot() const;
+
+	/** Set the behavior implementation to use */
+	void SetBehaviorImplementation(const TSharedRef<ICEClonerAttachmentTreeBehavior>& InImplementation);
+
+	/** Performs cleanup before destruction */
+	void Cleanup();
+
+	/** Marks the tree as outdated to perform an update */
+	void MarkAttachmentOutdated();
+
+	/** Marks the actor cache as outdated for an update */
+	void MarkCacheOutdated(AActor* InActor);
+	
+	/**
+	* Triggers an update of the attachment tree to detect updated items
+	* If reset is true, clears the attachment tree and rebuilds it otherwise diff update
+	*/
+	bool UpdateAttachments(bool bInReset = false);
+
+	/** Finds root actor of an actor based on this tree hierarchy by climbing up */
+	AActor* FindRootActor(AActor* InActor);
+
+	/** Retrieves all root actors from this tree hierarchy */
+	TArray<AActor*> GetRootActors() const;
+
+	/** Traverses each attachment in tree, stops when false is returned */
+	void ForEachAttachment(const TFunctionRef<bool(AActor*, FCEClonerAttachmentItem&)>& InFunctor);
+
+	/** Detach an item from the tree */
+	void DetachItem(AActor* InActor);
+
+	/** True if all root meshes are valid and usable */
+	bool IsCacheAvailable(bool bInAllowInvalid = true) const;
+	
+	/** Retrieves all attachments based on an actor */
+	void GetAttachments(AActor* InActor, TArray<const FCEClonerAttachmentItem*>& OutAttachmentItems, bool bInRecurse) const;
+
+	/** Get the attachment bounds based on tree hierarchy */
+	FBox GetAttachmentBounds(AActor* InActor, bool bInIncludeChildren) const;
+
+	/** Get cached bounds for a root actor */
+	FBox GetCacheBounds(AActor* InActor) const;
+
+	/** All cloner attached actor items */
+	UPROPERTY()
+	TMap<TWeakObjectPtr<AActor>, FCEClonerAttachmentItem> ItemAttachmentMap;
+
+	/** Actors directly attached to the cloner actor, order is important here */
+	UPROPERTY()
+	TArray<TWeakObjectPtr<AActor>> RootActors;
+
+	/**
+	 * Merged static meshes corresponding to the root actors for niagara,
+	 * result of merging dynamic meshes together
+	 */
+	UPROPERTY()
+	TArray<TObjectPtr<UStaticMesh>> MergedBakedMeshes;
+
+	/** Status of the cloner tree */
+	UPROPERTY(Transient)
+	ECEClonerAttachmentStatus Status = ECEClonerAttachmentStatus::Outdated;
+
+	/** Attachment items that are dirty and need an update */
+	UPROPERTY(Transient)
+	TSet<TWeakObjectPtr<AActor>> DirtyItemAttachments;
+
+	/** Whether Added/Moved/Removed items requires an update */
+	UPROPERTY(Transient)
+	bool bItemAttachmentsDirty = false;
+
+private:
+	bool UpdateAttachmentsInternal();
+	void UpdateAttachment(AActor* InActor, AActor* InParent, const FTransform& InActorTransform);
+	void DetachItem(FCEClonerAttachmentItem& InItem);
+
+	/** Clears the whole tree */
+	void Reset();
+
+	/** Called when tree item is attached */
+	FOnAttachmentChanged OnItemAttachedDelegate;
+
+	/** Called when tree item is detached */
+	FOnAttachmentChanged OnItemDetachedDelegate;
+
+	TObjectPtr<AActor> AttachmentRoot;
+
+	TSharedPtr<ICEClonerAttachmentTreeBehavior> BehaviorImplementation;
+};
diff --git a/Engine/Plugins/Experimental/ClonerEffector/Source/ClonerEffector/Public/Cloner/Attachments/CEClonerAttachmentTreeBehavior.h b/Engine/Plugins/Experimental/ClonerEffector/Source/ClonerEffector/Public/Cloner/Attachments/CEClonerAttachmentTreeBehavior.h
new file mode 100644
--- /dev/null
+++ b/Engine/Plugins/Experimental/ClonerEffector/Source/ClonerEffector/Public/Cloner/Attachments/CEClonerAttachmentTreeBehavior.h
@@ -0,0 +1,84 @@
+// Copyright Epic Games, Inc. All Rights Reserved.
+
+#pragma once
+
+#include "Components/SceneComponent.h"
+#include "Containers/Array.h"
+#include "Delegates/IDelegateInstance.h"
+#include "Templates/SharedPointer.h"
+
+class AActor;
+struct FCEClonerAttachmentItem;
+struct FCEClonerAttachmentTree;
+
+/** Defines a custom implementation for the attachment tree */
+class ICEClonerAttachmentTreeBehavior : public TSharedFromThis<ICEClonerAttachmentTreeBehavior>
+{
+public:
+	virtual ~ICEClonerAttachmentTreeBehavior() = default;
+
+	/** Called when this implementation is activated */
+	virtual void OnActivation(FCEClonerAttachmentTree& InTree) = 0;
+
+	/** Called when this implementation is deactivated */
+	virtual void OnDeactivation(FCEClonerAttachmentTree& InTree) = 0;
+
+	/** Retrieve ordered list of children actors */
+	virtual void GetOrderedChildrenActors(FCEClonerAttachmentTree& InTree, AActor* InActor, TArray<AActor*>& OutChildren) const = 0;
+
+	/** Attachment mesh status is out of date and needs an update */
+	virtual void InvalidateMesh(FCEClonerAttachmentTree& InTree, AActor* InActor) = 0;
+
+	/** Item is attached in the tree */
+	virtual void OnItemAttached(FCEClonerAttachmentTree& InTree, FCEClonerAttachmentItem& InItem) = 0;
+
+	/** Item is detached from the tree */
+	virtual void OnItemDetached(FCEClonerAttachmentTree& InTree, FCEClonerAttachmentItem& InItem) = 0;
+};
+
+/** Base hierarchy where children are grouped under a root cloner actor */
+class FCEClonerAttachmentGroupBehavior : public ICEClonerAttachmentTreeBehavior
+{
+protected:
+	//~ Begin ICEClonerAttachmentImplementation
+	virtual void OnActivation(FCEClonerAttachmentTree& InTree) override;
+	virtual void OnDeactivation(FCEClonerAttachmentTree& InTree) override;
+	virtual void GetOrderedChildrenActors(FCEClonerAttachmentTree& InTree, AActor* InActor, TArray<AActor*>& OutChildren) const override;
+	virtual void InvalidateMesh(FCEClonerAttachmentTree& InTree, AActor* InActor) override;
+	virtual void OnItemAttached(FCEClonerAttachmentTree& InTree, FCEClonerAttachmentItem& InItem) override;
+	virtual void OnItemDetached(FCEClonerAttachmentTree& InTree, FCEClonerAttachmentItem& InItem) override;
+	//~ End ICEClonerAttachmentImplementation
+
+	void BindActorDelegates(FCEClonerAttachmentTree& InTree, AActor* InActor);
+	void UnbindActorDelegates(FCEClonerAttachmentTree& InTree, AActor* InActor);
+
+	/** Called when the transform state of a cloned component changes */
+	void OnComponentTransformed(FCEClonerAttachmentTree& InTree, USceneComponent* InComponent, EUpdateTransformFlags InFlags, ETeleportType InTeleport);
+
+	/** Called when an actor in the world is destroyed */
+	void OnWorldActorDestroyed(FCEClonerAttachmentTree& InTree, AActor* InActor);
+
+#if WITH_EDITOR
+	/** Called when a mesh has changed in any cloned actors */
+	void OnMeshChanged(FCEClonerAttachmentTree& InTree, AActor* InActor);
+
+	/** Called when the render state of a cloned component changes */
+	void OnRenderStateDirty(FCEClonerAttachmentTree& InTree, UActorComponent& InComponent);
+
+	/** Called when material is done compiling and when property is changed to test whether material have changed */
+	void OnCheckMaterialChanged(FCEClonerAttachmentTree& InTree, UObject* InObject);
+
+	/** Called when a level actor hierarchy changes */
+	void OnLevelHierarchyChanged(FCEClonerAttachmentTree& InTree, AActor* InActor);
+#endif
+
+	FDelegateHandle WorldActorDestroyedDelegate;
+};
+
+/** Flat hierarchy where any grouped actor is unboxed as if it was under a root cloner actor */
+class FCEClonerAttachmentFlatBehavior : public FCEClonerAttachmentGroupBehavior
+{
+	//~ Begin ICEClonerAttachmentImplementation
+	virtual void GetOrderedChildrenActors(FCEClonerAttachmentTree& InTree, AActor* InActor, TArray<AActor*>& OutChildren) const override;
+	//~ End ICEClonerAttachmentImplementation
+};
diff --git a/Engine/Plugins/Experimental/ClonerEffector/Source/ClonerEffector/Public/Cloner/Attachments/CEClonerAttachmentTypes.h b/Engine/Plugins/Experimental/ClonerEffector/Source/ClonerEffector/Public/Cloner/Attachments/CEClonerAttachmentTypes.h
new file mode 100644
--- /dev/null
+++ b/Engine/Plugins/Experimental/ClonerEffector/Source/ClonerEffector/Public/Cloner/Attachments/CEClonerAttachmentTypes.h
@@ -0,0 +1,18 @@
+// Copyright Epic Games, Inc. All Rights Reserved.
+
+#pragma once
+
+#include "CEClonerAttachmentTypes.generated.h"
+
+UENUM()
+enum class ECEClonerAttachmentStatus : uint8
+{
+	/** Item should be removed, no longer valid */
+	Invalid,
+	/** Item should be updated, changes detected */
+	Outdated,
+	/** Item is up to date, no changes needed */
+	Updated,
+	/** Item is being updated at the moment */
+	Updating
+};
\ No newline at end of file
diff --git a/Engine/Plugins/Experimental/ClonerEffector/Source/ClonerEffector/Public/Cloner/Attachments/CEClonerSceneTreeCustomResolver.h b/Engine/Plugins/Experimental/ClonerEffector/Source/ClonerEffector/Public/Cloner/Attachments/CEClonerSceneTreeCustomResolver.h
new file mode 100644
--- /dev/null
+++ b/Engine/Plugins/Experimental/ClonerEffector/Source/ClonerEffector/Public/Cloner/Attachments/CEClonerSceneTreeCustomResolver.h
@@ -0,0 +1,23 @@
+// Copyright Epic Games, Inc. All Rights Reserved.
+
+#pragma once
+
+#include "Containers/Array.h"
+#include "Delegates/IDelegateInstance.h"
+#include "GameFramework/Actor.h"
+#include "HAL/Platform.h"
+#include "Templates/SharedPointer.h"
+
+class ULevel;
+
+/** Used to resolve underlying actors in the scene */
+class ICEClonerSceneTreeCustomResolver : public TSharedFromThis<ICEClonerSceneTreeCustomResolver>
+{
+public:
+	DECLARE_MULTICAST_DELEGATE_OneParam(FOnActorHierarchyChanged, AActor* /** InActor */)
+	virtual ~ICEClonerSceneTreeCustomResolver() = default;
+	virtual void Activate() = 0;
+	virtual void Deactivate() = 0;
+	virtual bool GetDirectChildrenActor(AActor* InActor, TArray<AActor*>& OutActors) const = 0;
+	virtual FOnActorHierarchyChanged::RegistrationType& OnActorHierarchyChanged() = 0;
+};
\ No newline at end of file
diff --git a/Engine/Plugins/Experimental/ClonerEffector/Source/ClonerEffector/Public/Cloner/CEClonerComponent.h b/Engine/Plugins/Experimental/ClonerEffector/Source/ClonerEffector/Public/Cloner/CEClonerComponent.h
--- a/Engine/Plugins/Experimental/ClonerEffector/Source/ClonerEffector/Public/Cloner/CEClonerComponent.h
+++ b/Engine/Plugins/Experimental/ClonerEffector/Source/ClonerEffector/Public/Cloner/CEClonerComponent.h
@@ -2,6 +2,7 @@
 
 #pragma once
 
+#include "Attachments/CEClonerAttachmentTree.h"
 #include "CEClonerEffectorShared.h"
 #include "CEMeshBuilder.h"
 #include "CEPropertyChangeDispatcher.h"
@@ -33,18 +34,7 @@ class UCEClonerComponent : public UNiagaraComponent
 	DECLARE_MULTICAST_DELEGATE_TwoParams(FOnClonerAttachmentChanged, UCEClonerComponent* /** ClonerComponent */, AActor* /** InActor */)
 
 public:
-	/** Only materials transient or part of the content folder can be dirtied, engine or plugins cannot */
-	static bool IsMaterialDirtyable(const UMaterialInterface* InMaterial);
-
-	/** Check if material has niagara usage flag set */
-	static bool IsMaterialUsageFlagSet(const UMaterialInterface* InMaterial);
-
-	static void SetActorVisibility(AActor* InActor, bool bInVisibility, ECEClonerActorVisibility InTarget = ECEClonerActorVisibility::All);
-
 #if WITH_EDITOR
-	/** Show material warning notification when missing niagara usage flag */
-	static void ShowMaterialWarning(int32 InMaterialCount);
-
 	static CLONEREFFECTOR_API FName GetActiveExtensionsPropertyName();
 
 	static CLONEREFFECTOR_API FName GetActiveLayoutPropertyName();
@@ -88,15 +78,6 @@ public:
 		return bEnabled;
 	}
 
-	UFUNCTION(BlueprintCallable, Category="Cloner")
-	CLONEREFFECTOR_API void SetTreeUpdateInterval(float InInterval);
-
-	UFUNCTION(BlueprintPure, Category="Cloner")
-	float GetTreeUpdateInterval() const
-	{
-		return TreeUpdateInterval;
-	}
-
 	UFUNCTION(BlueprintCallable, Category="Cloner")
 	CLONEREFFECTOR_API void SetSeed(int32 InSeed);
 
@@ -176,6 +157,15 @@ public:
 	}
 
 #if WITH_EDITOR
+	UFUNCTION(BlueprintCallable, Category="Cloner")
+	void SetTreeBehaviorName(FName InBehaviorName);
+
+	UFUNCTION(BlueprintPure, Category="Cloner")
+	FName GetTreeBehaviorName() const
+	{
+		return TreeBehaviorName;
+	}
+
 	UFUNCTION(BlueprintCallable, Category="Cloner")
 	CLONEREFFECTOR_API void SetVisualizerSpriteVisible(bool bInVisible);
 
@@ -280,21 +270,14 @@ protected:
 	/** Called when cloner actor is detached */
 	CLONEREFFECTOR_API static FOnClonerAttachmentChanged OnClonerActorDetachedDelegate;
 
-	/** Replaces all unsupported material by default material, gathers unset materials that needs recompiling with proper flags */
-	static bool FilterSupportedMaterials(TArray<TWeakObjectPtr<UMaterialInterface>>& InMaterials, TArray<TWeakObjectPtr<UMaterialInterface>>& OutUnsetMaterials, UMaterialInterface* InDefaultMaterial);
-
-	static bool FilterSupportedMaterial(UMaterialInterface*& InMaterial, UMaterialInterface* InDefaultMaterial);
-
-	/** Fires a warning about unset materials used within this cloner */
-	void FireMaterialWarning(const AActor* InContextActor, const TArray<TWeakObjectPtr<UMaterialInterface>>& InUnsetMaterials);
-
 	//~ Begin UObject
 	virtual void GetPreloadDependencies(TArray<UObject*>& OutDeps) override;
 	virtual void PostInitProperties() override;
 	virtual void PostLoad() override;
 #if WITH_EDITOR
-	virtual void PostEditImport() override;
+	virtual void PreDuplicate(FObjectDuplicationParameters& InParams) override;
 	virtual void PostDuplicate(bool bInPIE) override;
+	virtual void PostEditImport() override;
 	virtual void PostEditUndo() override;
 	virtual void PostEditChangeProperty(FPropertyChangedEvent& InPropertyChangedEvent) override;
 #endif
@@ -305,37 +288,25 @@ protected:
 	virtual void OnComponentDestroyed(bool bInDestroyingHierarchy) override;
 	//~ End UActorComponent
 
-	void UpdateAttachmentTree();
-	void UpdateActorAttachment(AActor* InActor, AActor* InParent);
-
-	/** Get cloner direct children in the correct order */
-	void GetOrderedRootActors(TArray<AActor*>& OutActors) const;
-
-	/** Get the root cloner actor (direct child) for a specific actor */
-	AActor* GetRootActor(AActor* InActor) const;
-
-	/** Resets the root baked static mesh to be regenerated later */
-	void InvalidateBakedStaticMesh(AActor* InActor);
-
+#if WITH_EDITOR
 	/** Runs async update to rebuild dirty meshes */
 	void UpdateDirtyMeshesAsync();
-	void OnDirtyMeshesUpdated(bool bInSuccess);
 
-	/** Merges all primitive components from an actor to dynamic mesh, does not recurse */
-	void UpdateActorBakedDynamicMesh(AActor* InActor);
+	/** Merges all primitive components from a single actor into a dynamic mesh */
+	void UpdateActorDynamicMesh(AActor* InActor);
 
-	/** Merges all baked dynamic meshes from children and self into one static mesh to use as niagara mesh input */
-	void UpdateRootActorBakedStaticMesh(AActor* InRootActor);
+	/** Merges all actors dynamic meshes into a single static mesh */
+	void UpdateActorStaticMesh(AActor* InRootActor);
 
-	/** Gets all attachment items based on an actor, will recurse */
-	void GetActorAttachmentItems(AActor* InActor, TArray<FCEClonerAttachmentItem*>& OutAttachmentItems);
+	/** Called once static mesh is done building */
+	void OnActorStaticMeshPostBuild(UStaticMesh* InMesh);
+#endif
+
+	void OnDirtyMeshesUpdated(bool bInSuccess);
 
 	/** Update niagara asset static meshes */
 	void UpdateClonerMeshes();
 
-	/** Checks that all root static meshes are valid */
-	bool IsAllMergedMeshesValid() const;
-
 	/** Sets the layout to use for this cloner simulation */
 	void SetClonerActiveLayout(UCEClonerLayoutBase* InLayout);
 
@@ -365,9 +336,9 @@ protected:
 	UPROPERTY(EditInstanceOnly, Setter, Getter, Category=General)
 	FRotator GlobalRotation = FRotator::ZeroRotator;
 
-	/** Interval to update the attachment tree and update the cloner meshes, 0 means each tick */
-	UPROPERTY(EditInstanceOnly, Setter, Getter, Category=General, AdvancedDisplay, meta=(ClampMin="0"))
-	float TreeUpdateInterval = 0.f;
+	/** Set the tree behavior with attachments */
+	UPROPERTY(EditInstanceOnly, DisplayName="Tree Mode", Category=General, meta=(GetOptions="GetClonerTreeBehaviorNames"))
+	FName TreeBehaviorName = NAME_None;
 
 #if WITH_EDITORONLY_DATA
 	/** Toggle the sprite to visualize and click on this cloner */
@@ -406,12 +377,6 @@ private:
 	bool CheckResourcesReady();
 	bool TickCloner(float InDelta);
 
-	/**
-	* Triggers an update of the attachment tree to detect updated items
-	* If reset is true, clears the attachment tree and rebuilds it otherwise diff update
-	*/
-	void UpdateClonerAttachmentTree(bool bInReset = false);
-
 	/** Called to trigger an update of cloner rendering state tree */
 	void UpdateClonerRenderState();
 
@@ -424,35 +389,15 @@ private:
 	void OnColorChanged();
 	void OnGlobalScaleChanged();
 	void OnGlobalRotationChanged();
+	void OnTreeBehaviorNameChanged();
 	void OnLayoutNameChanged();
 
 #if WITH_EDITOR
 	void OnVisualizerSpriteVisibleChanged();
 #endif
 
-	void BindActorDelegates(AActor* InActor);
-	void UnbindActorDelegates(AActor* InActor) const;
-
-#if WITH_EDITOR
-	void OnActorPropertyChanged(UObject* InObject, FPropertyChangedEvent& InPropertyChangedEvent);
-	void OnMaterialCompiled(UMaterialInterface* InMaterial);
-
-	/** Called when the render state of a cloned component changes */
-	void OnRenderStateDirty(UActorComponent& InActorComponent);
-#endif
-
-	/** Called when a cloned actor is destroyed */
-	UFUNCTION()
-	void OnActorDestroyed(AActor* InDestroyedActor);
-
-	/** Called when a material has changed in any cloned actors */
-	void OnMaterialChanged(UObject* InObject);
-
-	/** Called when a mesh has changed in any cloned actors */
-	void OnMeshChanged(UStaticMeshComponent*, AActor* InActor);
-
-	/** Called when the transform state of a cloned component changes */
-	void OnComponentTransformed(USceneComponent* InComponent, EUpdateTransformFlags InFlags, ETeleportType InTeleport);
+	void OnTreeItemAttached(AActor* InActor, FCEClonerAttachmentItem& InItem);
+	void OnTreeItemDetached(AActor* InActor, FCEClonerAttachmentItem& InItem);
 
 	template<
 		typename InLayoutClass
@@ -486,21 +431,21 @@ private:
 	UFUNCTION()
 	TArray<FName> GetClonerLayoutNames() const;
 
+	/** Gets all tree implementation names available */
+	UFUNCTION()
+	TArray<FName> GetClonerTreeBehaviorNames() const;
+
 	/** Attachment tree view */
 	UPROPERTY(NonTransactional)
 	FCEClonerAttachmentTree ClonerTree;
 
+	/** Mesh builder to create final meshes that are sent to niagara */
 	UPROPERTY(Transient, DuplicateTransient, NonTransactional)
 	FCEMeshBuilder MeshBuilder;
 
-	/** Asset meshes needs update */
-	bool bClonerMeshesDirty = true;
-
 	/** State of the baked dynamic and static mesh creation */
 	std::atomic<bool> bClonerMeshesUpdating = false;
 
-	float TreeUpdateDeltaTime = 0.f;
-
 	bool bNeedsRefresh = false;
 
 	bool bClonerInitialized = false;
@@ -510,8 +455,6 @@ private:
 	FTSTicker::FDelegateHandle ClonerTickerHandle;
 
 #if WITH_EDITOR
-	double LastNotificationTime = 0.0;
-
 	/** Used for PECP */
 	static const TCEPropertyChangeDispatcher<UCEClonerComponent> PropertyChangeDispatcher;
 #endif
diff --git a/Engine/Plugins/Experimental/ClonerEffector/Source/ClonerEffector/Public/Subsystems/CEClonerSubsystem.h b/Engine/Plugins/Experimental/ClonerEffector/Source/ClonerEffector/Public/Subsystems/CEClonerSubsystem.h
--- a/Engine/Plugins/Experimental/ClonerEffector/Source/ClonerEffector/Public/Subsystems/CEClonerSubsystem.h
+++ b/Engine/Plugins/Experimental/ClonerEffector/Source/ClonerEffector/Public/Subsystems/CEClonerSubsystem.h
@@ -3,11 +3,17 @@
 #pragma once
 
 #include "Cloner/Layouts/CEClonerLayoutBase.h"
+#include "Containers/Map.h"
+#include "Engine/Level.h"
 #include "Subsystems/EngineSubsystem.h"
+#include "Templates/SharedPointerFwd.h"
+#include "UObject/ObjectKey.h"
 #include "UObject/WeakObjectPtrTemplates.h"
 #include "CEClonerSubsystem.generated.h"
 
 class AActor;
+class ICEClonerAttachmentTreeBehavior;
+class ICEClonerSceneTreeCustomResolver;
 class UCEClonerExtensionBase;
 
 UCLASS(MinimalAPI)
@@ -29,6 +35,12 @@ public:
 		return OnClonerSetEnabledDelegate;
 	}
 
+	DECLARE_DELEGATE_RetVal_OneParam(TSharedPtr<ICEClonerSceneTreeCustomResolver>, FOnGetSceneTreeResolver, ULevel* /** Level */)
+	static FOnGetSceneTreeResolver::RegistrationType& OnGetSceneTreeResolver()
+	{
+		return OnGetSceneTreeResolverDelegate;
+	}
+
 	/** Get this subsystem instance */
 	CLONEREFFECTOR_API static UCEClonerSubsystem* Get();
 
@@ -53,13 +65,6 @@ public:
 	/** Creates a new layout instance for a cloner */
 	UCEClonerLayoutBase* CreateNewLayout(FName InLayoutName, UCEClonerComponent* InCloner);
 
-	DECLARE_DELEGATE_RetVal_OneParam(TArray<AActor*> /** Children */, FOnGetOrderedActors, const AActor* /** InParent */)
-	CLONEREFFECTOR_API void RegisterCustomActorResolver(FOnGetOrderedActors InCustomResolver);
-
-	CLONEREFFECTOR_API void UnregisterCustomActorResolver();
-
-	FOnGetOrderedActors& GetCustomActorResolver();
-
 	CLONEREFFECTOR_API bool RegisterExtensionClass(UClass* InClass);
 
 	CLONEREFFECTOR_API bool UnregisterExtensionClass(UClass* InClass);
@@ -95,18 +100,42 @@ public:
 	/** Creates a new cloner with actors attached */
 	CLONEREFFECTOR_API AActor* CreateClonerWithActors(UWorld* InWorld, const TSet<AActor*>& InActors, bool bInShouldTransact);
 
+	/** Fires a warning about unset materials used within a cloner */
+	void FireMaterialWarning(const AActor* InClonerActor, const AActor* InContextActor, const TArray<TWeakObjectPtr<UMaterialInterface>>& InUnsetMaterials);
+
+	/** Register an attachment tree behavior */
+	bool RegisterAttachmentTreeBehavior(FName InName, TFunction<TSharedRef<ICEClonerAttachmentTreeBehavior>()>& InCreator);
+
+	/** Unregister an attachment tree behavior */
+	bool UnregisterAttachmentTreeBehavior(FName InName);
+
+	/** Gets all attachment behavior registered */
+	TArray<FName> GetAttachmentTreeBehaviorNames() const;
+
+	/** Creates the specific attachment tree behavior */
+	TSharedPtr<ICEClonerAttachmentTreeBehavior> CreateAttachmentTreeBehavior(FName InName) const;
+
+	/** Finds a custom scene tree resolver for a specific level if available */
+	TSharedPtr<ICEClonerSceneTreeCustomResolver> FindCustomLevelSceneTreeResolver(ULevel* InLevel);
+
 protected:
 	CLONEREFFECTOR_API static FOnSubsystemInitialized OnSubsystemInitializedDelegate;
 
 	/** Delegate to change state of cloners in a world */
 	static FOnClonerSetEnabled OnClonerSetEnabledDelegate;
 
+	CLONEREFFECTOR_API static FOnGetSceneTreeResolver OnGetSceneTreeResolverDelegate;
+
 	//~ Begin USubsystem
 	virtual void Initialize(FSubsystemCollectionBase& Collection) override;
+	virtual void Deinitialize() override;
 	//~ End USubsystem
 
 	void ScanForRegistrableClasses();
 
+	void OnLevelCleanup(ULevel* InLevel);
+	void OnWorldCleanup(UWorld* InWorld, bool, bool bInCleanupResources);
+
 	/** Linking name to the layout class */
 	UPROPERTY()
 	TMap<FName, TSubclassOf<UCEClonerLayoutBase>> LayoutClasses;
@@ -115,6 +144,13 @@ protected:
 	UPROPERTY()
 	TMap<FName, TSubclassOf<UCEClonerExtensionBase>> ExtensionClasses;
 
+	/** Used to create a cloner tree attachment behavior */
+	TMap<FName, TFunction<TSharedRef<ICEClonerAttachmentTreeBehavior>()>> TreeBehaviorCreators;
+
 	/** Used to gather ordered actors based on parent */
-	FOnGetOrderedActors ActorResolver;
+	TMap<TObjectKey<ULevel>, TSharedRef<ICEClonerSceneTreeCustomResolver>> LevelCustomResolvers;
+
+#if WITH_EDITOR
+	double LastNotificationTime = 0.0;
+#endif
 };
\ No newline at end of file
diff --git a/Engine/Plugins/Experimental/ClonerEffector/Source/ClonerEffector/Public/Utilities/CEClonerEffectorUtilities.h b/Engine/Plugins/Experimental/ClonerEffector/Source/ClonerEffector/Public/Utilities/CEClonerEffectorUtilities.h
new file mode 100644
--- /dev/null
+++ b/Engine/Plugins/Experimental/ClonerEffector/Source/ClonerEffector/Public/Utilities/CEClonerEffectorUtilities.h
@@ -0,0 +1,39 @@
+// Copyright Epic Games, Inc. All Rights Reserved.
+
+#pragma once
+
+#include "Containers/Array.h"
+#include "UObject/WeakObjectPtr.h"
+#include "CEClonerEffectorShared.h"
+
+class AActor;
+class FText;
+class UMaterialInterface;
+
+namespace UE::ClonerEffector::Utilities
+{
+#if WITH_EDITOR
+	const FText& GetMaterialWarningText();
+
+	/** Show material warning notification when missing niagara usage flag */
+	void ShowWarning(const FText& InWarning);
+#endif
+
+	/** Only materials transient or part of the content folder can be dirtied, engine or plugins cannot */
+	bool IsMaterialDirtyable(const UMaterialInterface* InMaterial);
+
+	/** Check if material has niagara usage flag set */
+	bool IsMaterialUsageFlagSet(const UMaterialInterface* InMaterial);
+
+	/** Replaces all unsupported material by default material, gathers unset materials that needs recompiling with proper flags */
+	bool FilterSupportedMaterials(TArray<TWeakObjectPtr<UMaterialInterface>>& InMaterials, TArray<TWeakObjectPtr<UMaterialInterface>>& OutUnsetMaterials, UMaterialInterface* InDefaultMaterial);
+
+	/** Check if a material is supported otherwise replaces it with default material */
+	bool FilterSupportedMaterial(UMaterialInterface*& InMaterial, UMaterialInterface* InDefaultMaterial);
+
+	/** Set specific actor visibility */
+	void SetActorVisibility(AActor* InActor, bool bInVisibility, ECEClonerActorVisibility InTarget = ECEClonerActorVisibility::All);
+
+	/** Look parent attach actor until first cloner component is found */
+	AActor* FindClonerActor(AActor* InActor);
+}
