MotionDesign : ClonerEffector
- Disabled PVS studio warning since it is a false positive

diff --git a/Engine/Plugins/Experimental/ClonerEffector/Source/ClonerEffector/Private/Cloner/CEClonerComponent.cpp b/Engine/Plugins/Experimental/ClonerEffector/Source/ClonerEffector/Private/Cloner/CEClonerComponent.cpp
--- a/Engine/Plugins/Experimental/ClonerEffector/Source/ClonerEffector/Private/Cloner/CEClonerComponent.cpp
+++ b/Engine/Plugins/Experimental/ClonerEffector/Source/ClonerEffector/Private/Cloner/CEClonerComponent.cpp
@@ -372,7 +372,7 @@ void UCEClonerComponent::UpdateActorDynamicMesh(AActor* InActor)
 	AttachmentItem->BakedMaterials = MoveTemp(MeshMaterials);
 
 	// Was the mesh invalidated during the update process, then leave it outdated
-	if (AttachmentItem->MeshStatus == ECEClonerAttachmentStatus::Updating)
+	if (AttachmentItem->MeshStatus == ECEClonerAttachmentStatus::Updating) //-V547
 	{
 		AttachmentItem->MeshStatus = ECEClonerAttachmentStatus::Updated;
 		ClonerTree.MarkCacheOutdated(InActor);
