MotionDesign : ClonerEffector
- Fix grouped actors would not be rendered properly

diff --git a/Engine/Plugins/Experimental/ClonerEffector/Source/ClonerEffector/Private/Cloner/Attachments/CEClonerAttachmentTree.cpp b/Engine/Plugins/Experimental/ClonerEffector/Source/ClonerEffector/Private/Cloner/Attachments/CEClonerAttachmentTree.cpp
--- a/Engine/Plugins/Experimental/ClonerEffector/Source/ClonerEffector/Private/Cloner/Attachments/CEClonerAttachmentTree.cpp
+++ b/Engine/Plugins/Experimental/ClonerEffector/Source/ClonerEffector/Private/Cloner/Attachments/CEClonerAttachmentTree.cpp
@@ -99,7 +99,10 @@ bool FCEClonerAttachmentTree::UpdateAttachments(bool bInReset)
 	UE_LOG(LogCECloner, Log, TEXT("%s : updating attachment tree"), *AttachmentRoot->GetActorNameOrLabel())
 
 	Status = ECEClonerAttachmentStatus::Updating;
-	return UpdateAttachmentsInternal();
+	const bool bResult = UpdateAttachmentsInternal();
+	Status = ECEClonerAttachmentStatus::Updated;
+
+	return bResult;
 }
 
 AActor* FCEClonerAttachmentTree::FindRootActor(AActor* InActor)
@@ -355,8 +358,6 @@ bool FCEClonerAttachmentTree::UpdateAttachmentsInternal()
 
 	RootActors = NewRootActors;
 	MergedBakedMeshes = NewCombinesMeshes;
-	Status = ECEClonerAttachmentStatus::Updated;
-
 	return bItemAttachmentsDirty;
 }
 
@@ -459,9 +460,15 @@ void FCEClonerAttachmentTree::UpdateAttachment(AActor* InActor, AActor* InParent
 		}
 
 		AttachmentItem->ChildrenActors.Empty(ChildrenActors.Num());
+		Algo::Transform(ChildrenActors, AttachmentItem->ChildrenActors,
+			[](AActor* InActor)
+			{
+				return InActor;
+			}
+		);
+
 		for (AActor* ChildActor : ChildrenActors)
 		{
-			AttachmentItem->ChildrenActors.Add(ChildActor);
 			UpdateAttachment(ChildActor, InActor, InActorTransform);
 		}
 
diff --git a/Engine/Plugins/Experimental/ClonerEffector/Source/ClonerEffector/Private/Cloner/CEClonerComponent.cpp b/Engine/Plugins/Experimental/ClonerEffector/Source/ClonerEffector/Private/Cloner/CEClonerComponent.cpp
--- a/Engine/Plugins/Experimental/ClonerEffector/Source/ClonerEffector/Private/Cloner/CEClonerComponent.cpp
+++ b/Engine/Plugins/Experimental/ClonerEffector/Source/ClonerEffector/Private/Cloner/CEClonerComponent.cpp
@@ -348,8 +348,6 @@ void UCEClonerComponent::UpdateActorDynamicMesh(AActor* InActor)
 
 	AttachmentItem->MeshStatus = ECEClonerAttachmentStatus::Updating;
 
-	UE_LOG(LogCECloner, Log, TEXT("%s : Updating actor dynamic mesh %s"), *ClonerActor->GetActorNameOrLabel(), *InActor->GetActorNameOrLabel());
-
 	UDynamicMesh* Mesh = NewObject<UDynamicMesh>();
 	TArray<TWeakObjectPtr<UMaterialInterface>> MeshMaterials;
 
@@ -377,6 +375,8 @@ void UCEClonerComponent::UpdateActorDynamicMesh(AActor* InActor)
 		AttachmentItem->MeshStatus = ECEClonerAttachmentStatus::Updated;
 		ClonerTree.MarkCacheOutdated(InActor);
 	}
+
+	UE_LOG(LogCECloner, Log, TEXT("%s : Updated actor dynamic mesh : %s - %i"), *ClonerActor->GetActorNameOrLabel(), *InActor->GetActorNameOrLabel(), Mesh->GetTriangleCount());
 }
 
 void UCEClonerComponent::UpdateActorStaticMesh(AActor* InRootActor)
@@ -407,11 +407,10 @@ void UCEClonerComponent::UpdateActorStaticMesh(AActor* InRootActor)
 		return;
 	}
 
-	UE_LOG(LogCECloner, Log, TEXT("%s : Updating actor static mesh %s"), *ClonerActor->GetActorNameOrLabel(), *InRootActor->GetActorNameOrLabel());
-
 	TArray<const FCEClonerAttachmentItem*> AttachmentItems;
 	ClonerTree.GetAttachments(InRootActor, AttachmentItems, /** Recurse */true);
 
+	uint32 Hash = HashCombineFast(ClonerActor->GetUniqueID(), InRootActor->GetUniqueID());
 	for (const FCEClonerAttachmentItem* AttachmentItem : AttachmentItems)
 	{
 		if (!AttachmentItem)
@@ -425,12 +424,13 @@ void UCEClonerComponent::UpdateActorStaticMesh(AActor* InRootActor)
 			continue;
 		}
 
+		Hash = HashCombineFast(Hash, BakedDynamicMesh->GetUniqueID());
 		MeshBuilder.AppendMesh(AttachmentItem->BakedMesh, AttachmentItem->BakedMaterials, AttachmentItem->ActorTransform);
 	}
 
 	// Avoid dirtying asset by creating it transient first,
 	// and avoid bounds log spamming by renaming asset with prefix LandscapeNaniteMesh
-	const FName PreObjectName(FString::Printf(TEXT("LandscapeNaniteMesh_C%i_M%i"), ClonerActor->GetUniqueID(), InRootActor->GetUniqueID()));
+	const FName PreObjectName(FString::Printf(TEXT("LandscapeNaniteMesh_%i_%i"), Hash, AttachmentItems.Num()));
 	UStaticMesh* Mesh = NewObject<UStaticMesh>(GetTransientPackage(), PreObjectName);
 
 	// Rename the asset + outer once build is done
@@ -441,6 +441,8 @@ void UCEClonerComponent::UpdateActorStaticMesh(AActor* InRootActor)
 	MeshBuilder.Reset();
 
 	ClonerTree.MergedBakedMeshes[RootIdx] = Mesh;
+
+	UE_LOG(LogCECloner, Log, TEXT("%s : Updated actor static mesh : %s - %i - %i"), *ClonerActor->GetActorNameOrLabel(), *InRootActor->GetActorNameOrLabel(), Mesh->GetNumTriangles(/** LOD */0), AttachmentItems.Num());
 }
 
 void UCEClonerComponent::OnActorStaticMeshPostBuild(UStaticMesh* InMesh)
diff --git a/Engine/Plugins/Experimental/ClonerEffector/Source/ClonerEffectorMeshBuilder/Private/CEMeshBuilder.cpp b/Engine/Plugins/Experimental/ClonerEffector/Source/ClonerEffectorMeshBuilder/Private/CEMeshBuilder.cpp
--- a/Engine/Plugins/Experimental/ClonerEffector/Source/ClonerEffectorMeshBuilder/Private/CEMeshBuilder.cpp
+++ b/Engine/Plugins/Experimental/ClonerEffector/Source/ClonerEffectorMeshBuilder/Private/CEMeshBuilder.cpp
@@ -1052,6 +1052,7 @@ bool FCEMeshBuilder::DynamicMeshToStaticMesh(UDynamicMesh* InMesh, UStaticMesh*
 
 	// Mesh
 	FMeshDescription& MeshDescription = StaticMeshDescription->GetMeshDescription();
+	MeshDescription.Empty();
 
 	InMesh->ProcessMesh([&MeshDescription](const FDynamicMesh3& InSourceMesh)
 	{
