MotionDesign : ClonerEffector
- Add submenu to create cloner with specific layout or create linked effector with specific shape
- Fixed undo+redo issue causing effector to not work anymore when removed (undo) then recreated (redo)

diff --git a/Engine/Plugins/Experimental/Avalanche/Source/AvalancheEditor/Private/ClonerEffector/AvaClonerEffectorExtension.cpp b/Engine/Plugins/Experimental/Avalanche/Source/AvalancheEditor/Private/ClonerEffector/AvaClonerEffectorExtension.cpp
--- a/Engine/Plugins/Experimental/Avalanche/Source/AvalancheEditor/Private/ClonerEffector/AvaClonerEffectorExtension.cpp
+++ b/Engine/Plugins/Experimental/Avalanche/Source/AvalancheEditor/Private/ClonerEffector/AvaClonerEffectorExtension.cpp
@@ -147,9 +147,8 @@ void FAvaClonerEffectorExtension::CreateCloner() const
 	TSet<AActor*> SelectedActors = GetSelectedActors();
 	UWorld* World = ModeTools->GetWorld();
 
-	constexpr bool bTransact = true;
-
-	ClonerSubsystem->CreateClonerWithActors(World, SelectedActors, bTransact);
+	constexpr UCEClonerSubsystem::ECEClonerActionFlags Flags = UCEClonerSubsystem::ECEClonerActionFlags::All;
+	ClonerSubsystem->CreateClonerWithActors(World, SelectedActors, Flags);
 }
 
 #undef LOCTEXT_NAMESPACE
diff --git a/Engine/Plugins/Experimental/ClonerEffector/Source/ClonerEffector/Private/CEModule.cpp b/Engine/Plugins/Experimental/ClonerEffector/Source/ClonerEffector/Private/CEModule.cpp
--- a/Engine/Plugins/Experimental/ClonerEffector/Source/ClonerEffector/Private/CEModule.cpp
+++ b/Engine/Plugins/Experimental/ClonerEffector/Source/ClonerEffector/Private/CEModule.cpp
@@ -3,7 +3,9 @@
 #include "Modules/ModuleManager.h"
 
 #include "Cloner/Logs/CEClonerLogs.h"
+#include "Effector/Logs/CEEffectorLogs.h"
 
 DEFINE_LOG_CATEGORY(LogCECloner);
+DEFINE_LOG_CATEGORY(LogCEEffector);
 
 IMPLEMENT_MODULE(FDefaultModuleImpl, ClonerEffector)
\ No newline at end of file
diff --git a/Engine/Plugins/Experimental/ClonerEffector/Source/ClonerEffector/Private/Cloner/Extensions/CEClonerEffectorExtension.cpp b/Engine/Plugins/Experimental/ClonerEffector/Source/ClonerEffector/Private/Cloner/Extensions/CEClonerEffectorExtension.cpp
--- a/Engine/Plugins/Experimental/ClonerEffector/Source/ClonerEffector/Private/Cloner/Extensions/CEClonerEffectorExtension.cpp
+++ b/Engine/Plugins/Experimental/ClonerEffector/Source/ClonerEffector/Private/Cloner/Extensions/CEClonerEffectorExtension.cpp
@@ -3,11 +3,10 @@
 #include "Cloner/Extensions/CEClonerEffectorExtension.h"
 
 #include "Cloner/CEClonerComponent.h"
-#include "Cloner/Layouts/CEClonerLayoutBase.h"
 #include "Cloner/Logs/CEClonerLogs.h"
-#include "Effector/CEEffectorActor.h"
 #include "Effector/CEEffectorComponent.h"
 #include "NiagaraDataInterfaceArrayInt.h"
+#include "Subsystems/CEClonerSubsystem.h"
 #include "Subsystems/CEEffectorSubsystem.h"
 
 #if WITH_EDITOR
@@ -223,38 +222,13 @@ void UCEClonerEffectorExtension::OnEffectorActorsChanged()
 #if WITH_EDITOR
 void UCEClonerEffectorExtension::CreateLinkedEffector()
 {
-	const UCEClonerComponent* ClonerComponent = GetClonerComponent();
-	if (!IsValid(ClonerComponent))
-	{
-		return;
-	}
-
-	UWorld* ClonerWorld = ClonerComponent->GetWorld();
-	if (!IsValid(ClonerWorld))
-	{
-		return;
-	}
-
-	FActorSpawnParameters Params;
-	Params.bTemporaryEditorActor = false;
-
-	const FVector ClonerLocation = ClonerComponent->GetComponentLocation();
-	const FRotator ClonerRotation = ClonerComponent->GetComponentRotation();
-
-	ACEEffectorActor* EffectorActor = ClonerWorld->SpawnActor<ACEEffectorActor>(ACEEffectorActor::StaticClass(), ClonerLocation, ClonerRotation, Params);
-	if (!EffectorActor)
-	{
-		return;
-	}
-
-	FActorLabelUtilities::RenameExistingActor(EffectorActor, EffectorActor->GetDefaultActorLabel(), true);
-
-	LinkEffector(EffectorActor);
+	UCEClonerComponent* ClonerComponent = GetClonerComponent();
+	UCEClonerSubsystem* ClonerSubsystem = UCEClonerSubsystem::Get();
 
-	if (GEditor)
+	if (ClonerComponent && ClonerSubsystem)
 	{
-		GEditor->SelectNone(true, true);
-		GEditor->SelectActor(EffectorActor, true, true);
+		constexpr UCEClonerSubsystem::ECEClonerActionFlags Flags = UCEClonerSubsystem::ECEClonerActionFlags::All;
+		ClonerSubsystem->CreateLinkedEffectors({ClonerComponent}, Flags);
 	}
 }
 
diff --git a/Engine/Plugins/Experimental/ClonerEffector/Source/ClonerEffector/Private/Effector/CEEffectorActor.cpp b/Engine/Plugins/Experimental/ClonerEffector/Source/ClonerEffector/Private/Effector/CEEffectorActor.cpp
--- a/Engine/Plugins/Experimental/ClonerEffector/Source/ClonerEffector/Private/Effector/CEEffectorActor.cpp
+++ b/Engine/Plugins/Experimental/ClonerEffector/Source/ClonerEffector/Private/Effector/CEEffectorActor.cpp
@@ -14,14 +14,13 @@
 #include "Effector/Types/CEEffectorRadialType.h"
 #include "Effector/Types/CEEffectorSphereType.h"
 #include "Effector/Types/CEEffectorTorusType.h"
+#include "Logs/CEEffectorLogs.h"
 
 #if WITH_EDITOR
 #include "Materials/MaterialInstanceDynamic.h"
 #include "UObject/ConstructorHelpers.h"
 #endif
 
-DEFINE_LOG_CATEGORY_STATIC(LogCEEffectorActor, Log, All);
-
 struct FCEEffectorActorVersion
 {
 	enum Type : int32
@@ -98,7 +97,7 @@ void ACEEffectorActor::MigrateDeprecatedProperties()
 		return;
 	}
 
-	UE_LOG(LogCEEffectorActor, Warning, TEXT("%s : Effector migrating from version %i to latest %i, please re-save this asset"), *GetActorNameOrLabel(), MigrateToVersion, FCEEffectorActorVersion::LatestVersion);
+	UE_LOG(LogCEEffector, Warning, TEXT("%s : Effector migrating from version %i to latest %i, please re-save this asset"), *GetActorNameOrLabel(), MigrateToVersion, FCEEffectorActorVersion::LatestVersion);
 
 	PRAGMA_DISABLE_DEPRECATION_WARNINGS
 
diff --git a/Engine/Plugins/Experimental/ClonerEffector/Source/ClonerEffector/Private/Effector/CEEffectorComponent.cpp b/Engine/Plugins/Experimental/ClonerEffector/Source/ClonerEffector/Private/Effector/CEEffectorComponent.cpp
--- a/Engine/Plugins/Experimental/ClonerEffector/Source/ClonerEffector/Private/Effector/CEEffectorComponent.cpp
+++ b/Engine/Plugins/Experimental/ClonerEffector/Source/ClonerEffector/Private/Effector/CEEffectorComponent.cpp
@@ -13,13 +13,11 @@
 #include "Effector/Types/CEEffectorTypeBase.h"
 #include "Engine/Texture2D.h"
 #include "GameFramework/Actor.h"
+#include "Logs/CEEffectorLogs.h"
 #include "Materials/MaterialInstanceDynamic.h"
 #include "Materials/MaterialInterface.h"
-#include "Settings/CEClonerEffectorSettings.h"
 #include "Subsystems/CEEffectorSubsystem.h"
 
-DEFINE_LOG_CATEGORY_STATIC(LogCEEffectorComponent, Log, All);
-
 #if WITH_EDITOR
 FName UCEEffectorComponent::GetModeNamePropertyName()
 {
@@ -170,9 +168,7 @@ FCEClonerEffectorChannelData& UCEEffectorComponent::GetChannelData()
 
 void UCEEffectorComponent::RegisterToChannel()
 {
-	PRAGMA_DISABLE_DEPRECATION_WARNINGS
-	if (IsThisNotNull(this, "UCEEffectorComponent::RegisterToChannel") && IsValidChecked(this) && ChannelData.GetIdentifier() == INDEX_NONE)
-	PRAGMA_ENABLE_DEPRECATION_WARNINGS
+	if (IsValid(this) && ChannelData.GetIdentifier() == INDEX_NONE)
 	{
 		// Register this effector to the effector channel
 		if (UCEEffectorSubsystem* EffectorSubsystem = UCEEffectorSubsystem::Get())
@@ -182,6 +178,15 @@ void UCEEffectorComponent::RegisterToChannel()
 	}
 }
 
+void UCEEffectorComponent::UnregisterFromChannel()
+{
+	// Remove this effector from the effector channel
+	if (UCEEffectorSubsystem* EffectorSubsystem = UCEEffectorSubsystem::Get())
+	{
+		EffectorSubsystem->UnregisterChannelEffector(this);
+	}
+}
+
 int32 UCEEffectorComponent::GetChannelIdentifier() const
 {
 	return ChannelData.GetIdentifier();
@@ -288,11 +293,7 @@ void UCEEffectorComponent::OnComponentDestroyed(bool bDestroyingHierarchy)
 {
 	Super::OnComponentDestroyed(bDestroyingHierarchy);
 
-	// Remove this effector from the effector channel
-	if (UCEEffectorSubsystem* EffectorSubsystem = UCEEffectorSubsystem::Get())
-	{
-		EffectorSubsystem->UnregisterChannelEffector(this);
-	}
+	UnregisterFromChannel();
 }
 
 void UCEEffectorComponent::PostEditImport()
@@ -305,7 +306,7 @@ void UCEEffectorComponent::PostEditImport()
 		{
 			if (UCEClonerEffectorExtension* ClonerExtension = ClonerExtensionWeak.Get())
 			{
-				UE_LOG(LogCEEffectorComponent, Log, TEXT("Linking effector %s to cloner %s after duplication"), *Owner->GetActorNameOrLabel(), *ClonerExtension->GetClonerComponent()->GetOwner()->GetActorNameOrLabel())
+				UE_LOG(LogCEEffector, Log, TEXT("Linking effector %s to cloner %s after duplication"), *Owner->GetActorNameOrLabel(), *ClonerExtension->GetClonerComponent()->GetOwner()->GetActorNameOrLabel())
 
 				ClonerExtension->LinkEffector(Owner);
 			}
@@ -358,7 +359,14 @@ void UCEEffectorComponent::PostEditUndo()
 {
 	Super::PostEditUndo();
 
-	OnEffectorOptionsChanged();
+	if (IsValid(this))
+	{
+		OnEffectorOptionsChanged();
+	}
+	else
+	{
+		UnregisterFromChannel();
+	}
 }
 
 void UCEEffectorComponent::SetVisualizerComponentVisible(bool bInVisible)
diff --git a/Engine/Plugins/Experimental/ClonerEffector/Source/ClonerEffector/Private/Effector/Logs/CEEffectorLogs.h b/Engine/Plugins/Experimental/ClonerEffector/Source/ClonerEffector/Private/Effector/Logs/CEEffectorLogs.h
new file mode 100644
--- /dev/null
+++ b/Engine/Plugins/Experimental/ClonerEffector/Source/ClonerEffector/Private/Effector/Logs/CEEffectorLogs.h
@@ -0,0 +1,7 @@
+// Copyright Epic Games, Inc. All Rights Reserved.
+
+#pragma once
+
+#include "Logging/LogMacros.h"
+
+DECLARE_LOG_CATEGORY_EXTERN(LogCEEffector, Log, All);
\ No newline at end of file
diff --git a/Engine/Plugins/Experimental/ClonerEffector/Source/ClonerEffector/Private/Subsystems/CEClonerSubsystem.cpp b/Engine/Plugins/Experimental/ClonerEffector/Source/ClonerEffector/Private/Subsystems/CEClonerSubsystem.cpp
--- a/Engine/Plugins/Experimental/ClonerEffector/Source/ClonerEffector/Private/Subsystems/CEClonerSubsystem.cpp
+++ b/Engine/Plugins/Experimental/ClonerEffector/Source/ClonerEffector/Private/Subsystems/CEClonerSubsystem.cpp
@@ -19,7 +19,9 @@
 #include "Cloner/Layouts/CEClonerSphereUniformLayout.h"
 #include "Cloner/Layouts/CEClonerSplineLayout.h"
 #include "Cloner/Logs/CEClonerLogs.h"
+#include "Effector/CEEffectorActor.h"
 #include "Engine/Engine.h"
+#include "Engine/World.h"
 #include "GameFramework/Actor.h"
 #include "Materials/Material.h"
 #include "UObject/Class.h"
@@ -28,6 +30,7 @@
 
 #if WITH_EDITOR
 #include "Editor.h"
+#include "Editor/EditorEngine.h"
 #include "ScopedTransaction.h"
 #endif
 
@@ -364,61 +367,126 @@ void UCEClonerSubsystem::ConvertCloners(const TSet<UCEClonerComponent*>& InClone
 		}
 	}
 }
+#endif
 
-void UCEClonerSubsystem::CreateLinkedEffector(const TSet<UCEClonerComponent*>& InCloners)
+TArray<UCEEffectorComponent*> UCEClonerSubsystem::CreateLinkedEffectors(const TArray<UCEClonerComponent*>& InCloners, ECEClonerActionFlags InFlags, TFunctionRef<void(UCEEffectorComponent*)> InGenerator)
 {
+	TArray<UCEEffectorComponent*> LinkedEffectors;
+
 	if (InCloners.IsEmpty())
 	{
-		return;
+		return LinkedEffectors;
 	}
 
-	for (UCEClonerComponent* ClonerComponent : InCloners)
+	LinkedEffectors.Reserve(InCloners.Num());
+
+#if WITH_EDITOR
+	const bool bSelect = EnumHasAnyFlags(InFlags, ECEClonerActionFlags::ShouldSelect);
+	const bool bShouldTransact = EnumHasAnyFlags(InFlags, ECEClonerActionFlags::ShouldTransact);
+
+	if (bSelect && GEditor)
+	{
+		GEditor->SelectNone(/** Notify */false, /** DeselectBSP */true);
+	}
+
+	const FText TransactionText = FText::Format(LOCTEXT("CreateLinkedEffectors", "Create {0} linked effector(s)"), FText::AsNumber(InCloners.Num()));
+	FScopedTransaction Transaction(TransactionText, bShouldTransact && !GIsTransacting);
+#endif
+
+	for (const UCEClonerComponent* ClonerComponent : InCloners)
 	{
 		if (!IsValid(ClonerComponent))
 		{
 			continue;
 		}
 
+		UWorld* ClonerWorld = ClonerComponent->GetWorld();
+		if (!IsValid(ClonerWorld))
+		{
+			continue;
+		}
+
+		FActorSpawnParameters SpawnParameters;
+		SpawnParameters.OverrideLevel = ClonerComponent->GetComponentLevel();
+#if WITH_EDITOR
+		SpawnParameters.InitialActorLabel = ACEEffectorActor::DefaultLabel;
+		SpawnParameters.ObjectFlags = RF_Transactional;
+		SpawnParameters.bTemporaryEditorActor = false;
+#endif
+
+		const FVector ClonerLocation = ClonerComponent->GetComponentLocation();
+		const FRotator ClonerRotation = ClonerComponent->GetComponentRotation();
+
+		ACEEffectorActor* EffectorActor = ClonerWorld->SpawnActor<ACEEffectorActor>(ACEEffectorActor::StaticClass(), ClonerLocation, ClonerRotation, SpawnParameters);
+		if (!EffectorActor)
+		{
+			continue;
+		}
+
 		if (UCEClonerEffectorExtension* EffectorExtension = ClonerComponent->GetExtension<UCEClonerEffectorExtension>())
 		{
-			EffectorExtension->CreateLinkedEffector();
+#if WITH_EDITOR
+			EffectorExtension->Modify();
+#endif
+
+			EffectorExtension->LinkEffector(EffectorActor);
+		}
+
+		InGenerator(LinkedEffectors.Add_GetRef(EffectorActor->GetEffectorComponent()));
+
+#if WITH_EDITOR
+		if (bSelect && GEditor)
+		{
+			GEditor->SelectActor(EffectorActor, /** Selected */true, /** Notify */true);
 		}
+#endif
+	}
+
+#if WITH_EDITOR
+	if (LinkedEffectors.IsEmpty())
+	{
+		Transaction.Cancel();
 	}
-}
 #endif
 
-AActor* UCEClonerSubsystem::CreateClonerWithActors(UWorld* InWorld, const TSet<AActor*>& InActors, bool bInShouldTransact)
+	return LinkedEffectors;
+}
+
+UCEClonerComponent* UCEClonerSubsystem::CreateClonerWithActors(UWorld* InWorld, const TSet<AActor*>& InActors, ECEClonerActionFlags InFlags)
 {
-	ACEClonerActor* NewClonerActor = nullptr;
+	UCEClonerComponent* NewCloner = nullptr;
 
 	if (!IsValid(InWorld))
 	{
-		return NewClonerActor;
+		return NewCloner;
 	}
 
 #if WITH_EDITOR
-	FScopedTransaction Transaction(LOCTEXT("CreateClonerWithActors", "Create cloner with actors attached"), bInShouldTransact);
+	const bool bSelect = EnumHasAnyFlags(InFlags, ECEClonerActionFlags::ShouldSelect);
+	const bool bShouldTransact = EnumHasAnyFlags(InFlags, ECEClonerActionFlags::ShouldTransact);
+
+	FScopedTransaction Transaction(
+		LOCTEXT("CreateClonerWithActors", "Create cloner with actors attached"),
+		bShouldTransact && !GIsTransacting
+	);
 #endif
 
 	FActorSpawnParameters Parameters;
-	Parameters.ObjectFlags = RF_Transactional;
 #if WITH_EDITOR
+	Parameters.InitialActorLabel = ACEClonerActor::DefaultLabel;
+	Parameters.ObjectFlags = RF_Transactional;
 	Parameters.bTemporaryEditorActor = false;
 #endif
 
-	NewClonerActor = InWorld->SpawnActor<ACEClonerActor>(Parameters);
-
-	if (NewClonerActor)
+	if (ACEClonerActor* NewClonerActor = InWorld->SpawnActor<ACEClonerActor>(Parameters))
 	{
-#if WITH_EDITOR
-		NewClonerActor->Modify();
-#endif
+		NewCloner = NewClonerActor->GetClonerComponent();
 
 		if (!InActors.IsEmpty())
 		{
 			FVector NewAverageLocation;
 
-			for (AActor* Actor : InActors)
+			for (const AActor* Actor : InActors)
 			{
 				if (IsValid(Actor))
 				{
@@ -442,7 +510,7 @@ AActor* UCEClonerSubsystem::CreateClonerWithActors(UWorld* InWorld, const TSet<A
 		}
 
 #if WITH_EDITOR
-		if (GEditor)
+		if (bSelect && GEditor)
 		{
 			GEditor->SelectNone(/** SelectionChange */false, /** DeselectBSP */true);
 			GEditor->SelectActor(NewClonerActor, /** Selected */true, /** Notify */true);
@@ -450,7 +518,7 @@ AActor* UCEClonerSubsystem::CreateClonerWithActors(UWorld* InWorld, const TSet<A
 #endif
 	}
 
-	return NewClonerActor;
+	return NewCloner;
 }
 
 void UCEClonerSubsystem::FireMaterialWarning(const AActor* InClonerActor, const AActor* InContextActor, const TArray<TWeakObjectPtr<UMaterialInterface>>& InUnsetMaterials)
diff --git a/Engine/Plugins/Experimental/ClonerEffector/Source/ClonerEffector/Private/Subsystems/CEEffectorSubsystem.cpp b/Engine/Plugins/Experimental/ClonerEffector/Source/ClonerEffector/Private/Subsystems/CEEffectorSubsystem.cpp
--- a/Engine/Plugins/Experimental/ClonerEffector/Source/ClonerEffector/Private/Subsystems/CEEffectorSubsystem.cpp
+++ b/Engine/Plugins/Experimental/ClonerEffector/Source/ClonerEffector/Private/Subsystems/CEEffectorSubsystem.cpp
@@ -2,12 +2,9 @@
 
 #include "Subsystems/CEEffectorSubsystem.h"
 
-#include "Engine/Engine.h"
-#include "NiagaraDataChannel.h"
-#include "NiagaraDataChannelAccessor.h"
-#include "NiagaraDataChannelPublic.h"
 #include "Effector/CEEffectorComponent.h"
 #include "Effector/CEEffectorExtensionBase.h"
+#include "Effector/Logs/CEEffectorLogs.h"
 #include "Effector/Modes/CEEffectorOffsetMode.h"
 #include "Effector/Modes/CEEffectorProceduralMode.h"
 #include "Effector/Modes/CEEffectorPushMode.h"
@@ -18,13 +15,15 @@
 #include "Effector/Types/CEEffectorSphereType.h"
 #include "Effector/Types/CEEffectorTorusType.h"
 #include "Effector/Types/CEEffectorUnboundType.h"
+#include "Engine/Engine.h"
+#include "NiagaraDataChannel.h"
+#include "NiagaraDataChannelAccessor.h"
+#include "NiagaraDataChannelPublic.h"
 
 #if WITH_EDITOR
 #include "ScopedTransaction.h"
 #endif
 
-DEFINE_LOG_CATEGORY_STATIC(LogACEEffectorSubsystem, Log, All);
-
 UCEEffectorSubsystem::FOnSubsystemInitialized UCEEffectorSubsystem::OnSubsystemInitializedDelegate;
 UCEEffectorSubsystem::FOnEffectorIdentifierChanged UCEEffectorSubsystem::OnEffectorIdentifierChangedDelegate;
 UCEEffectorSubsystem::FOnEffectorSetEnabled UCEEffectorSubsystem::OnEffectorSetEnabledDelegate;
@@ -81,7 +80,7 @@ bool UCEEffectorSubsystem::RegisterChannelEffector(UCEEffectorComponent* InEffec
 	if (EffectorIndex == INDEX_NONE)
 	{
 		EffectorIndex = EffectorsWeak.Add(InEffector);
-		UE_LOG(LogACEEffectorSubsystem, Log, TEXT("%s effector registered in channel %i"), *InEffector->GetOwner()->GetActorNameOrLabel(), EffectorIndex);
+		UE_LOG(LogCEEffector, Log, TEXT("%s effector registered in channel %i"), *InEffector->GetOwner()->GetActorNameOrLabel(), EffectorIndex);
 	}
 
 	if (InEffector->GetChannelData().Identifier != EffectorIndex)
@@ -105,7 +104,7 @@ bool UCEEffectorSubsystem::UnregisterChannelEffector(UCEEffectorComponent* InEff
 
 	if (bUnregistered)
 	{
-		UE_LOG(LogACEEffectorSubsystem, Log, TEXT("%s effector unregistered from channel"), *InEffector->GetOwner()->GetActorNameOrLabel());
+		UE_LOG(LogCEEffector, Log, TEXT("%s effector unregistered from channel"), *InEffector->GetOwner()->GetActorNameOrLabel());
 
 		const int32 OldIdentifier = InEffector->GetChannelData().Identifier;
 		InEffector->GetChannelData().Identifier = INDEX_NONE;
@@ -331,7 +330,7 @@ void UCEEffectorSubsystem::UpdateEffectorChannel(const UWorld* InWorld)
 
 	if (!ChannelWriter)
 	{
-		UE_LOG(LogACEEffectorSubsystem, Warning, TEXT("Effector data channel writer is invalid"));
+		UE_LOG(LogCEEffector, Warning, TEXT("Effector data channel writer is invalid"));
 		return;
 	}
 
diff --git a/Engine/Plugins/Experimental/ClonerEffector/Source/ClonerEffector/Public/Cloner/CEClonerComponent.h b/Engine/Plugins/Experimental/ClonerEffector/Source/ClonerEffector/Public/Cloner/CEClonerComponent.h
--- a/Engine/Plugins/Experimental/ClonerEffector/Source/ClonerEffector/Public/Cloner/CEClonerComponent.h
+++ b/Engine/Plugins/Experimental/ClonerEffector/Source/ClonerEffector/Public/Cloner/CEClonerComponent.h
@@ -224,8 +224,7 @@ public:
 	/** Forces a refresh of the meshes used */
 	void RefreshClonerMeshes();
 
-	template<
-		typename InExtensionClass>
+	template<typename InExtensionClass UE_REQUIRES(std::is_base_of_v<UCEClonerExtensionBase, InExtensionClass>)>
 	InExtensionClass* GetExtension() const
 	{
 		return Cast<InExtensionClass>(GetExtension(InExtensionClass::StaticClass()));
diff --git a/Engine/Plugins/Experimental/ClonerEffector/Source/ClonerEffector/Public/Effector/CEEffectorComponent.h b/Engine/Plugins/Experimental/ClonerEffector/Source/ClonerEffector/Public/Effector/CEEffectorComponent.h
--- a/Engine/Plugins/Experimental/ClonerEffector/Source/ClonerEffector/Public/Effector/CEEffectorComponent.h
+++ b/Engine/Plugins/Experimental/ClonerEffector/Source/ClonerEffector/Public/Effector/CEEffectorComponent.h
@@ -175,6 +175,12 @@ public:
 	UFUNCTION(BlueprintCallable, Category="Effector")
 	UCEEffectorExtensionBase* GetExtension(TSubclassOf<UCEEffectorExtensionBase> InExtensionClass) const;
 
+	template<typename InExtensionClass UE_REQUIRES(std::is_base_of_v<UCEEffectorExtensionBase, InExtensionClass>)>
+	InExtensionClass* GetExtension() const
+	{
+		return Cast<InExtensionClass>(GetExtension(InExtensionClass::StaticClass()));
+	}
+
 	UCEEffectorExtensionBase* GetExtension(FName InExtensionName) const;
 
 	void RequestClonerUpdate(bool bInImmediate);
@@ -248,6 +254,7 @@ private:
 	static constexpr TCHAR SpriteTexturePath[] = TEXT("/Script/Engine.Texture2D'/ClonerEffector/Textures/T_EffectorIcon.T_EffectorIcon'");
 
 	void RegisterToChannel();
+	void UnregisterFromChannel();
 	void OnTransformUpdated(USceneComponent* InComponent, EUpdateTransformFlags InFlags, ETeleportType InTeleportType);
 
 	void OnEnabledChanged();
diff --git a/Engine/Plugins/Experimental/ClonerEffector/Source/ClonerEffector/Public/Subsystems/CEClonerSubsystem.h b/Engine/Plugins/Experimental/ClonerEffector/Source/ClonerEffector/Public/Subsystems/CEClonerSubsystem.h
--- a/Engine/Plugins/Experimental/ClonerEffector/Source/ClonerEffector/Public/Subsystems/CEClonerSubsystem.h
+++ b/Engine/Plugins/Experimental/ClonerEffector/Source/ClonerEffector/Public/Subsystems/CEClonerSubsystem.h
@@ -11,6 +11,7 @@
 #include "UObject/WeakObjectPtrTemplates.h"
 #include "CEClonerSubsystem.generated.h"
 
+class UCEEffectorComponent;
 class AActor;
 class ICEClonerAttachmentTreeBehavior;
 class ICEClonerSceneTreeCustomResolver;
@@ -25,6 +26,14 @@ class UCEClonerSubsystem : public UEngineSubsystem
 	DECLARE_MULTICAST_DELEGATE_ThreeParams(FOnClonerSetEnabled, const UWorld* /** InWorld */, bool /** bInEnabled */, bool /** bInTransact */)
 
 public:
+	enum class ECEClonerActionFlags : uint8
+	{
+		None,
+		ShouldTransact = 1 << 0,
+		ShouldSelect = 1 << 1,
+		All = ShouldTransact | ShouldSelect,
+	};
+
 	static FOnSubsystemInitialized::RegistrationType& OnSubsystemInitialized()
 	{
 		return OnSubsystemInitializedDelegate;
@@ -92,13 +101,13 @@ public:
 #if WITH_EDITOR
 	/** Converts cloners simulation to a mesh */
 	CLONEREFFECTOR_API void ConvertCloners(const TSet<UCEClonerComponent*>& InCloners, ECEClonerMeshConversion InMeshConversion);
-
-	/** Create cloners linked effector */
-	CLONEREFFECTOR_API void CreateLinkedEffector(const TSet<UCEClonerComponent*>& InCloners);
 #endif
 
+	/** Spawn linked effectors with a generator and optionally transact */
+	CLONEREFFECTOR_API TArray<UCEEffectorComponent*> CreateLinkedEffectors(const TArray<UCEClonerComponent*>& InCloners, ECEClonerActionFlags InFlags, TFunctionRef<void(UCEEffectorComponent*)> InGenerator = [](UCEEffectorComponent*){});
+
 	/** Creates a new cloner with actors attached */
-	CLONEREFFECTOR_API AActor* CreateClonerWithActors(UWorld* InWorld, const TSet<AActor*>& InActors, bool bInShouldTransact);
+	CLONEREFFECTOR_API UCEClonerComponent* CreateClonerWithActors(UWorld* InWorld, const TSet<AActor*>& InActors, ECEClonerActionFlags InFlags);
 
 	/** Fires a warning about unset materials used within a cloner */
 	void FireMaterialWarning(const AActor* InClonerActor, const AActor* InContextActor, const TArray<TWeakObjectPtr<UMaterialInterface>>& InUnsetMaterials);
diff --git a/Engine/Plugins/Experimental/ClonerEffector/Source/ClonerEffectorEditor/Private/Cloner/Menus/CEEditorClonerMenu.cpp b/Engine/Plugins/Experimental/ClonerEffector/Source/ClonerEffectorEditor/Private/Cloner/Menus/CEEditorClonerMenu.cpp
--- a/Engine/Plugins/Experimental/ClonerEffector/Source/ClonerEffectorEditor/Private/Cloner/Menus/CEEditorClonerMenu.cpp
+++ b/Engine/Plugins/Experimental/ClonerEffector/Source/ClonerEffectorEditor/Private/Cloner/Menus/CEEditorClonerMenu.cpp
@@ -2,9 +2,15 @@
 
 #include "Cloner/Menus/CEEditorClonerMenu.h"
 
+#include "Cloner/CEClonerActor.h"
 #include "Cloner/CEClonerComponent.h"
+#include "Effector/CEEffectorActor.h"
+#include "Effector/CEEffectorComponent.h"
 #include "Engine/World.h"
+#include "Styling/SlateIconFinder.h"
 #include "Subsystems/CEClonerSubsystem.h"
+#include "Utilities/CEClonerLibrary.h"
+#include "Utilities/CEEffectorLibrary.h"
 
 #define LOCTEXT_NAMESPACE "CEEditorClonerMenu"
 
@@ -111,15 +117,25 @@ void UE::ClonerEditor::Menu::FillCreateClonerEffectorSection(UToolMenu* InMenu,
 
 	check(ClonerSection)
 
-	ClonerSection->AddMenuEntry(
-		TEXT("CreateClonerLinkedEffector")
-		, LOCTEXT("CreateClonerLinkedEffector.Label", "Create linked effector")
-		, LOCTEXT("CreateClonerLinkedEffector.Tooltip", "Create linked effector for selected cloners")
-		, FSlateIcon()
-		, FUIAction(
-			FExecuteAction::CreateLambda(&ExecuteCreateClonerEffectorAction, InMenuData)
-		)
-	);
+	TSet<FName> EffectorTypes;
+	UCEEffectorLibrary::GetEffectorTypeNames(EffectorTypes);
+
+	for (const FName& EffectorType : EffectorTypes)
+	{
+		const FName MenuName(TEXT("CreateClonerLinkedEffector") + EffectorType.ToString());
+		const FText MenuLabel = FText::Format(LOCTEXT("CreateClonerLinkedEffector.Label", "Create {0} effector"), FText::FromName(EffectorType));
+		const FText MenuTooltip = FText::Format(LOCTEXT("CreateClonerLinkedEffector.Tooltip", "Create a linked {0} effector for selected cloners"), FText::FromName(EffectorType));
+		
+		ClonerSection->AddMenuEntry(
+			MenuName
+			, MenuLabel
+			, MenuTooltip
+			, FSlateIconFinder::FindIconForClass(ACEEffectorActor::StaticClass())
+			, FUIAction(
+				FExecuteAction::CreateLambda(&ExecuteCreateClonerEffectorAction, InMenuData, EffectorType)
+			)
+		);
+	}
 }
 
 void UE::ClonerEditor::Menu::FillConvertClonerSection(UToolMenu* InMenu, const FCEEditorClonerMenuData& InMenuData)
@@ -200,15 +216,25 @@ void UE::ClonerEditor::Menu::FillCreateClonerSection(UToolMenu* InMenu, const FC
 
 	check(ClonerSection)
 
-	ClonerSection->AddMenuEntry(
-		TEXT("CreateCloner")
-		, LOCTEXT("CreateCloner.Label", "Create cloner")
-		, LOCTEXT("CreateCloner.Tooltip", "Create cloner with selection attached")
-		, FSlateIcon()
-		, FUIAction(
-			FExecuteAction::CreateLambda(&ExecuteCreateClonerAction, InMenuData)
-		)
-	);
+	TSet<FName> ClonerLayouts;
+	UCEClonerLibrary::GetClonerLayoutNames(ClonerLayouts);
+
+	for (const FName& ClonerLayout : ClonerLayouts)
+	{
+		const FName MenuName(TEXT("CreateCloner") + ClonerLayout.ToString());
+		const FText MenuLabel = FText::Format(LOCTEXT("CreateCloner.Label", "Create {0} cloner"), FText::FromName(ClonerLayout));
+		const FText MenuTooltip = FText::Format(LOCTEXT("CreateCloner.Tooltip", "Create a {0} cloner with current selection attached"), FText::FromName(ClonerLayout));
+
+		ClonerSection->AddMenuEntry(
+			MenuName
+			, MenuLabel
+			, MenuTooltip
+			, FSlateIconFinder::FindIconForClass(ACEClonerActor::StaticClass())
+			, FUIAction(
+				FExecuteAction::CreateLambda(&ExecuteCreateClonerAction, InMenuData, ClonerLayout)
+			)
+		);
+	}
 }
 
 void UE::ClonerEditor::Menu::ExecuteEnableClonerAction(const FCEEditorClonerMenuData& InMenuData, bool bInEnable)
@@ -236,7 +262,7 @@ void UE::ClonerEditor::Menu::ExecuteEnableLevelClonerAction(const FCEEditorClone
 	Subsystem->SetLevelClonersEnabled(World, bInEnable, InMenuData.Options.ShouldTransact());
 }
 
-void UE::ClonerEditor::Menu::ExecuteCreateClonerEffectorAction(const FCEEditorClonerMenuData& InMenuData)
+void UE::ClonerEditor::Menu::ExecuteCreateClonerEffectorAction(const FCEEditorClonerMenuData& InMenuData, FName InEffectorType)
 {
 	UCEClonerSubsystem* Subsystem = UCEClonerSubsystem::Get();
 
@@ -245,7 +271,21 @@ void UE::ClonerEditor::Menu::ExecuteCreateClonerEffectorAction(const FCEEditorCl
 		return;
 	}
 
-	Subsystem->CreateLinkedEffector(InMenuData.Context.GetCloners());
+	UCEClonerSubsystem::ECEClonerActionFlags Flags = UCEClonerSubsystem::ECEClonerActionFlags::ShouldSelect;
+
+	if (InMenuData.Options.ShouldTransact())
+	{
+		EnumAddFlags(Flags, UCEClonerSubsystem::ECEClonerActionFlags::ShouldTransact);
+	}
+
+	Subsystem->CreateLinkedEffectors(
+		InMenuData.Context.GetCloners().Array(),
+		Flags,
+		[InEffectorType](UCEEffectorComponent* InEffector)
+		{
+			InEffector->SetTypeName(InEffectorType);
+		}
+	);
 }
 
 void UE::ClonerEditor::Menu::ExecuteConvertClonerAction(const FCEEditorClonerMenuData& InMenuData, ECEClonerMeshConversion InToMeshType)
@@ -260,7 +300,7 @@ void UE::ClonerEditor::Menu::ExecuteConvertClonerAction(const FCEEditorClonerMen
 	Subsystem->ConvertCloners(InMenuData.Context.GetEnabledCloners(), InToMeshType);
 }
 
-void UE::ClonerEditor::Menu::ExecuteCreateClonerAction(const FCEEditorClonerMenuData& InMenuData)
+void UE::ClonerEditor::Menu::ExecuteCreateClonerAction(const FCEEditorClonerMenuData& InMenuData, FName InClonerLayout)
 {
 	UCEClonerSubsystem* Subsystem = UCEClonerSubsystem::Get();
 
@@ -269,7 +309,17 @@ void UE::ClonerEditor::Menu::ExecuteCreateClonerAction(const FCEEditorClonerMenu
 		return;
 	}
 
-	Subsystem->CreateClonerWithActors(InMenuData.Context.GetWorld(), InMenuData.Context.GetActors(), InMenuData.Options.ShouldTransact());
+	UCEClonerSubsystem::ECEClonerActionFlags Flags = UCEClonerSubsystem::ECEClonerActionFlags::ShouldSelect;
+
+	if (InMenuData.Options.ShouldTransact())
+	{
+		EnumAddFlags(Flags, UCEClonerSubsystem::ECEClonerActionFlags::ShouldTransact);
+	}
+
+	if (UCEClonerComponent* Cloner = Subsystem->CreateClonerWithActors(InMenuData.Context.GetWorld(), InMenuData.Context.GetActors(), Flags))
+	{
+		Cloner->SetLayoutName(InClonerLayout);
+	}
 }
 
 #undef LOCTEXT_NAMESPACE
diff --git a/Engine/Plugins/Experimental/ClonerEffector/Source/ClonerEffectorEditor/Private/Cloner/Menus/CEEditorClonerMenu.h b/Engine/Plugins/Experimental/ClonerEffector/Source/ClonerEffectorEditor/Private/Cloner/Menus/CEEditorClonerMenu.h
--- a/Engine/Plugins/Experimental/ClonerEffector/Source/ClonerEffectorEditor/Private/Cloner/Menus/CEEditorClonerMenu.h
+++ b/Engine/Plugins/Experimental/ClonerEffector/Source/ClonerEffectorEditor/Private/Cloner/Menus/CEEditorClonerMenu.h
@@ -42,9 +42,9 @@ namespace UE::ClonerEditor::Menu
 
 	void ExecuteEnableLevelClonerAction(const FCEEditorClonerMenuData& InMenuData, bool bInEnable);
 
-	void ExecuteCreateClonerEffectorAction(const FCEEditorClonerMenuData& InMenuData);
+	void ExecuteCreateClonerEffectorAction(const FCEEditorClonerMenuData& InMenuData, FName InEffectorType);
 
 	void ExecuteConvertClonerAction(const FCEEditorClonerMenuData& InMenuData, ECEClonerMeshConversion InToMeshType);
 
-	void ExecuteCreateClonerAction(const FCEEditorClonerMenuData& InMenuData);
+	void ExecuteCreateClonerAction(const FCEEditorClonerMenuData& InMenuData, FName InClonerLayout);
 }
diff --git a/Engine/Plugins/Experimental/ClonerEffector/Source/ClonerEffectorEditor/Private/Subsystems/CEEditorClonerSubsystem.cpp b/Engine/Plugins/Experimental/ClonerEffector/Source/ClonerEffectorEditor/Private/Subsystems/CEEditorClonerSubsystem.cpp
--- a/Engine/Plugins/Experimental/ClonerEffector/Source/ClonerEffectorEditor/Private/Subsystems/CEEditorClonerSubsystem.cpp
+++ b/Engine/Plugins/Experimental/ClonerEffector/Source/ClonerEffectorEditor/Private/Subsystems/CEEditorClonerSubsystem.cpp
@@ -77,7 +77,20 @@ void UCEEditorClonerSubsystem::FillClonerMenu(UToolMenu* InMenu, const FCEEditor
 
 	if (InOptions.IsMenuType(ECEEditorClonerMenuType::CreateEffector) && InContext.ContainsAnyCloner())
 	{
-		FillCreateClonerEffectorSection(InMenu, MenuData);
+		if (InOptions.ShouldCreateSubMenu())
+		{
+			check(ClonerSection)
+
+			ClonerSection->AddSubMenu(
+				TEXT("CreateEffectorMenu"),
+				LOCTEXT("CreateEffectorMenu.Label", "Create effectors"),
+				LOCTEXT("CreateEffectorMenu.Tooltip", "Create linked effectors for selected cloners"),
+				FNewToolMenuDelegate::CreateLambda(&FillCreateClonerEffectorSection, MenuData));
+		}
+		else
+		{
+			FillCreateClonerEffectorSection(InMenu, MenuData);
+		}
 	}
 
 	if (InOptions.IsMenuType(ECEEditorClonerMenuType::Convert) && InContext.ContainsAnyEnabledCloner())
@@ -100,7 +113,20 @@ void UCEEditorClonerSubsystem::FillClonerMenu(UToolMenu* InMenu, const FCEEditor
 
 	if (InOptions.IsMenuType(ECEEditorClonerMenuType::CreateCloner) && !InContext.ContainsAnyCloner() && InContext.ContainsAnyActor())
 	{
-		FillCreateClonerSection(InMenu, MenuData);
+		if (InOptions.ShouldCreateSubMenu())
+		{
+			check(ClonerSection)
+
+			ClonerSection->AddSubMenu(
+				TEXT("CreateClonerMenu"),
+				LOCTEXT("CreateClonerMenu.Label", "Create cloner"),
+				LOCTEXT("CreateClonerMenu.Tooltip", "Create cloner with selected actors"),
+				FNewToolMenuDelegate::CreateLambda(&FillCreateClonerSection, MenuData));
+		}
+		else
+		{
+			FillCreateClonerSection(InMenu, MenuData);
+		}
 	}
 }
 
