MotionDesign : ClonerEffector
- Fix cloner and effector rotation issue with quaternion dritfting over time
- Changing namespace of effector module attributes to output instead of particle
- Add particles spawn frame/total limit to avoid cloner performance impact due to large particles count spawn
- Fix vortex force origin moved from emitter position to effector position

#virtualized

diff --git a/Engine/Plugins/Experimental/ClonerEffector/Source/ClonerEffector/Private/Cloner/Extensions/CEClonerEmitterSpawnExtension.cpp b/Engine/Plugins/Experimental/ClonerEffector/Source/ClonerEffector/Private/Cloner/Extensions/CEClonerEmitterSpawnExtension.cpp
--- a/Engine/Plugins/Experimental/ClonerEffector/Source/ClonerEffector/Private/Cloner/Extensions/CEClonerEmitterSpawnExtension.cpp
+++ b/Engine/Plugins/Experimental/ClonerEffector/Source/ClonerEffector/Private/Cloner/Extensions/CEClonerEmitterSpawnExtension.cpp
@@ -72,6 +72,30 @@ void UCEClonerEmitterSpawnExtension::SetSpawnRate(float InRate)
 	MarkExtensionDirty();
 }
 
+void UCEClonerEmitterSpawnExtension::SetSpawnMaxFrameCount(int32 InCount)
+{
+	InCount = FMath::Clamp(InCount, 0, SpawnMaxTotalCount);
+	if (SpawnMaxFrameCount == InCount)
+	{
+		return;
+	}
+
+	SpawnMaxFrameCount = InCount;
+	MarkExtensionDirty();
+}
+
+void UCEClonerEmitterSpawnExtension::SetSpawnMaxTotalCount(int32 InCount)
+{
+	InCount = FMath::Max(InCount, 0);
+	if (SpawnMaxTotalCount == InCount)
+	{
+		return;
+	}
+
+	SpawnMaxTotalCount = InCount;
+	MarkExtensionDirty();
+}
+
 void UCEClonerEmitterSpawnExtension::OnExtensionParametersChanged(UCEClonerComponent* InComponent)
 {
 	Super::OnExtensionParametersChanged(InComponent);
@@ -91,6 +115,10 @@ void UCEClonerEmitterSpawnExtension::OnExtensionParametersChanged(UCEClonerCompo
 
 	InComponent->SetFloatParameter(TEXT("SpawnRate"), SpawnRate);
 
+	InComponent->SetIntParameter(TEXT("SpawnMaxFrameCount"), FMath::Clamp(SpawnMaxFrameCount, 0, SpawnMaxTotalCount));
+
+	InComponent->SetIntParameter(TEXT("SpawnMaxTotalCount"), FMath::Max(SpawnMaxTotalCount, 0));
+
 #if WITH_EDITOR
 	// Do not allow world space when spawning once
 	if (SpawnLoopMode == ECEClonerSpawnLoopMode::Once && !bUseLocalSpace)
@@ -123,6 +151,8 @@ const TCEPropertyChangeDispatcher<UCEClonerEmitterSpawnExtension> UCEClonerEmitt
 	{ GET_MEMBER_NAME_CHECKED(UCEClonerEmitterSpawnExtension, SpawnBehaviorMode), &UCEClonerEmitterSpawnExtension::OnExtensionPropertyChanged },
 	{ GET_MEMBER_NAME_CHECKED(UCEClonerEmitterSpawnExtension, SpawnRate), &UCEClonerEmitterSpawnExtension::OnExtensionPropertyChanged },
 	{ GET_MEMBER_NAME_CHECKED(UCEClonerEmitterSpawnExtension, bUseLocalSpace), &UCEClonerEmitterSpawnExtension::OnExtensionPropertyChanged },
+	{ GET_MEMBER_NAME_CHECKED(UCEClonerEmitterSpawnExtension, SpawnMaxFrameCount), &UCEClonerEmitterSpawnExtension::OnExtensionPropertyChanged },
+	{ GET_MEMBER_NAME_CHECKED(UCEClonerEmitterSpawnExtension, SpawnMaxTotalCount), &UCEClonerEmitterSpawnExtension::OnExtensionPropertyChanged },
 };
 
 void UCEClonerEmitterSpawnExtension::PostEditChangeProperty(FPropertyChangedEvent& InPropertyChangedEvent)
@@ -159,4 +189,4 @@ void UCEClonerEmitterSpawnExtension::OnLocalSpaceChanged()
 		}
 	}
 }
-#endif
\ No newline at end of file
+#endif
diff --git a/Engine/Plugins/Experimental/ClonerEffector/Source/ClonerEffector/Public/Cloner/Extensions/CEClonerEmitterSpawnExtension.h b/Engine/Plugins/Experimental/ClonerEffector/Source/ClonerEffector/Public/Cloner/Extensions/CEClonerEmitterSpawnExtension.h
--- a/Engine/Plugins/Experimental/ClonerEffector/Source/ClonerEffector/Public/Cloner/Extensions/CEClonerEmitterSpawnExtension.h
+++ b/Engine/Plugins/Experimental/ClonerEffector/Source/ClonerEffector/Public/Cloner/Extensions/CEClonerEmitterSpawnExtension.h
@@ -63,6 +63,24 @@ public:
 		return SpawnRate;
 	}
 
+	UFUNCTION(BlueprintCallable, Category="Cloner")
+	CLONEREFFECTOR_API void SetSpawnMaxFrameCount(int32 InCount);
+
+	UFUNCTION(BlueprintPure, Category="Cloner")
+	int32 GetSpawnMaxFrameCount() const
+	{
+		return SpawnMaxFrameCount;
+	}
+	
+	UFUNCTION(BlueprintCallable, Category="Cloner")
+	CLONEREFFECTOR_API void SetSpawnMaxTotalCount(int32 InCount);
+
+	UFUNCTION(BlueprintPure, Category="Cloner")
+	int32 GetSpawnMaxTotalCount() const
+	{
+		return SpawnMaxTotalCount;
+	}
+
 #if WITH_EDITOR
 	UFUNCTION(BlueprintCallable, Category="Cloner")
 	CLONEREFFECTOR_API void SetUseLocalSpace(bool bInLocalSpace);
@@ -109,6 +127,14 @@ protected:
 	UPROPERTY(EditInstanceOnly, Category="Spawn", meta=(EditCondition="SpawnLoopMode != ECEClonerSpawnLoopMode::Once", EditConditionHides))
 	bool bUseLocalSpace = true;
 
+	/** Amount of particle allowed to spawn in a single frame */
+	UPROPERTY(EditInstanceOnly, Setter, Getter, DisplayName="Max Frame Spawn Count", Category="Spawn", meta=(ClampMin="0"))
+	int32 SpawnMaxFrameCount = 1000000;
+
+	/** Amount of particle allowed to spawn by the emitter */
+	UPROPERTY(EditInstanceOnly, Setter, Getter, DisplayName="Max Total Spawn Count", Category="Spawn", meta=(ClampMin="0"))
+	int32 SpawnMaxTotalCount = 1000000;
+
 private:
 #if WITH_EDITOR
 	/** Used for PECP */
diff --git a/Engine/Plugins/Experimental/ClonerEffector/Source/ClonerEffector/Public/Effector/Effects/CEEffectorDelayEffect.h b/Engine/Plugins/Experimental/ClonerEffector/Source/ClonerEffector/Public/Effector/Effects/CEEffectorDelayEffect.h
--- a/Engine/Plugins/Experimental/ClonerEffector/Source/ClonerEffector/Public/Effector/Effects/CEEffectorDelayEffect.h
+++ b/Engine/Plugins/Experimental/ClonerEffector/Source/ClonerEffector/Public/Effector/Effects/CEEffectorDelayEffect.h
@@ -83,7 +83,7 @@ protected:
 	
 	/** Time delay for the effect to reach its peak */
 	UPROPERTY(EditInstanceOnly, Setter, Getter, DisplayName="Offset Duration", Category="Delay", meta=(Units=Seconds, ClampMin="0", EditCondition="bDelayEnabled", EditConditionHides))
-	float DelayInDuration = 1.f;
+	float DelayInDuration = 0.f;
 
 	/** Time delay for the effect to go back to its rest state */
 	UPROPERTY(EditInstanceOnly, Setter, Getter, DisplayName="Trail Duration", Category="Delay", meta=(Units=Seconds, ClampMin="0", EditCondition="bDelayEnabled", EditConditionHides))
