MotionDesign : ClonerEffector
- Fix cloner default attached actor rotation will be zero in world space to avoid double rotation
- Fix crash happening when undo/redo is used after updating DefaultMeshes array property

diff --git a/Engine/Plugins/Experimental/ClonerEffector/Source/ClonerEffector/Private/Cloner/CEClonerComponent.cpp b/Engine/Plugins/Experimental/ClonerEffector/Source/ClonerEffector/Private/Cloner/CEClonerComponent.cpp
--- a/Engine/Plugins/Experimental/ClonerEffector/Source/ClonerEffector/Private/Cloner/CEClonerComponent.cpp
+++ b/Engine/Plugins/Experimental/ClonerEffector/Source/ClonerEffector/Private/Cloner/CEClonerComponent.cpp
@@ -844,7 +844,7 @@ void UCEClonerComponent::CreateDefaultActorAttached()
 	SpawnParameters.bTemporaryEditorActor = false;
 
 	const FVector ClonerLocation = GetComponentLocation();
-	const FRotator ClonerRotation = GetComponentRotation();
+	const FRotator ClonerRotation = FRotator::ZeroRotator;
 
 	if (AStaticMeshActor* DefaultActorAttached = World->SpawnActor<AStaticMeshActor>(ClonerLocation, ClonerRotation, SpawnParameters))
 	{
diff --git a/Engine/Plugins/Experimental/ClonerEffector/Source/ClonerEffector/Private/Cloner/Extensions/CEClonerMeshRendererExtension.cpp b/Engine/Plugins/Experimental/ClonerEffector/Source/ClonerEffector/Private/Cloner/Extensions/CEClonerMeshRendererExtension.cpp
--- a/Engine/Plugins/Experimental/ClonerEffector/Source/ClonerEffector/Private/Cloner/Extensions/CEClonerMeshRendererExtension.cpp
+++ b/Engine/Plugins/Experimental/ClonerEffector/Source/ClonerEffector/Private/Cloner/Extensions/CEClonerMeshRendererExtension.cpp
@@ -11,6 +11,11 @@
 #include "UObject/ConstructorHelpers.h"
 #include "Utilities/CEClonerEffectorUtilities.h"
 
+#if WITH_EDITOR
+#include "Containers/Ticker.h"
+#include "Misc/TransactionObjectEvent.h"
+#endif
+
 UCEClonerMeshRendererExtension::UCEClonerMeshRendererExtension()
 	: UCEClonerExtensionBase(
 		TEXT("MeshRenderer")
@@ -277,6 +282,40 @@ const TCEPropertyChangeDispatcher<UCEClonerMeshRendererExtension> UCEClonerMeshR
 	{ GET_MEMBER_NAME_CHECKED(UCEClonerMeshRendererExtension, bSortTranslucentParticles), &UCEClonerMeshRendererExtension::OnOverrideMaterialOptionsChanged },
 };
 
+void UCEClonerMeshRendererExtension::PostTransacted(const FTransactionObjectEvent& InTransactionEvent)
+{
+	Super::PostTransacted(InTransactionEvent);
+
+	if (InTransactionEvent.GetEventType() == ETransactionObjectEventType::UndoRedo)
+	{
+		if (InTransactionEvent.GetChangedProperties().Contains(GET_MEMBER_NAME_CHECKED(UCEClonerMeshRendererExtension, DefaultMeshes)))
+		{
+			if (UCEClonerComponent* ClonerComponent = GetClonerComponent())
+			{
+				// Redo : reactivate system and refresh
+				if (!ClonerComponent->IsActive())
+				{
+					ClonerComponent->SetActiveFlag(true);
+					FTSTicker::GetCoreTicker().AddTicker(FTickerDelegate::CreateWeakLambda(this, [this](float)
+					{
+						if (UCEClonerComponent* ClonerComponent = GetClonerComponent())
+						{
+							ClonerComponent->RefreshClonerMeshes();
+						}
+
+						return false;
+					}));
+				}
+				// Undo : deactivate system and destroy instance
+				else
+				{
+					ClonerComponent->DestroyInstanceNotComponent();
+				}
+			}
+		}
+	}
+}
+
 void UCEClonerMeshRendererExtension::PostEditChangeProperty(FPropertyChangedEvent& InPropertyChangedEvent)
 {
 	Super::PostEditChangeProperty(InPropertyChangedEvent);
diff --git a/Engine/Plugins/Experimental/ClonerEffector/Source/ClonerEffector/Public/Cloner/Extensions/CEClonerMeshRendererExtension.h b/Engine/Plugins/Experimental/ClonerEffector/Source/ClonerEffector/Public/Cloner/Extensions/CEClonerMeshRendererExtension.h
--- a/Engine/Plugins/Experimental/ClonerEffector/Source/ClonerEffector/Public/Cloner/Extensions/CEClonerMeshRendererExtension.h
+++ b/Engine/Plugins/Experimental/ClonerEffector/Source/ClonerEffector/Public/Cloner/Extensions/CEClonerMeshRendererExtension.h
@@ -96,6 +96,7 @@ public:
 protected:
 	//~ Begin UObject
 #if WITH_EDITOR
+	virtual void PostTransacted(const FTransactionObjectEvent& InTransactionEvent) override;
 	virtual void PostEditChangeProperty(FPropertyChangedEvent& InPropertyChangedEvent) override;
 #endif
 	//~ End UObject
@@ -150,4 +151,4 @@ private:
 	/** Used for PECP */
 	static const TCEPropertyChangeDispatcher<UCEClonerMeshRendererExtension> PropertyChangeDispatcher;
 #endif
-};
\ No newline at end of file
+};
