MotionDesign : ClonerEffector
- Added preserve ratio specifier to relevant vector properties

diff --git a/Engine/Plugins/VirtualProduction/ClonerEffector/Source/ClonerEffector/Public/Effector/Effects/CEEffectorForceEffect.h b/Engine/Plugins/VirtualProduction/ClonerEffector/Source/ClonerEffector/Public/Effector/Effects/CEEffectorForceEffect.h
--- a/Engine/Plugins/VirtualProduction/ClonerEffector/Source/ClonerEffector/Public/Effector/Effects/CEEffectorForceEffect.h
+++ b/Engine/Plugins/VirtualProduction/ClonerEffector/Source/ClonerEffector/Public/Effector/Effects/CEEffectorForceEffect.h
@@ -238,10 +238,10 @@ protected:
 	UPROPERTY(EditInstanceOnly, Setter, Getter, Category="Force", meta=(Delta="0.0001", ClampMin="0", EditCondition="bForcesEnabled && bOrientationForceEnabled", EditConditionHides))
 	float OrientationForceRate = 1.f;
 
-	UPROPERTY(EditInstanceOnly, Setter, Getter, Category="Force", meta=(EditCondition="bForcesEnabled && bOrientationForceEnabled", EditConditionHides))
+	UPROPERTY(EditInstanceOnly, Setter, Getter, Category="Force", meta=(AllowPreserveRatio, EditCondition="bForcesEnabled && bOrientationForceEnabled", EditConditionHides))
 	FVector OrientationForceMin = FVector(-0.1f);
 
-	UPROPERTY(EditInstanceOnly, Setter, Getter, Category="Force", meta=(EditCondition="bForcesEnabled && bOrientationForceEnabled", EditConditionHides))
+	UPROPERTY(EditInstanceOnly, Setter, Getter, Category="Force", meta=(AllowPreserveRatio, EditCondition="bForcesEnabled && bOrientationForceEnabled", EditConditionHides))
 	FVector OrientationForceMax = FVector(0.1f);
 
 	/** Enable vortex force to allow each clone instance to rotate around a specific axis on the cloner pivot */
