MotionDesign : ClonerEffector
- Fixed effector niagara data channel leak caused by the subsystem pushing data to a world that does not tick and increasing queue size, this is happening when PIE is started and editor world tick is halted, the niagara manager does not tick and does not consume requests
- Fixed preview neighbor grid mesh has player collision enabled limiting usage of cloner during game

diff --git a/Engine/Plugins/VirtualProduction/ClonerEffector/Source/ClonerEffector/Private/Cloner/Extensions/CEClonerCollisionExtension.cpp b/Engine/Plugins/VirtualProduction/ClonerEffector/Source/ClonerEffector/Private/Cloner/Extensions/CEClonerCollisionExtension.cpp
--- a/Engine/Plugins/VirtualProduction/ClonerEffector/Source/ClonerEffector/Private/Cloner/Extensions/CEClonerCollisionExtension.cpp
+++ b/Engine/Plugins/VirtualProduction/ClonerEffector/Source/ClonerEffector/Private/Cloner/Extensions/CEClonerCollisionExtension.cpp
@@ -243,6 +243,7 @@ bool UCEClonerCollisionExtension::InitVisualizerComponent()
 	CollisionVisualizerComponent->SetIsVisualizationComponent(true);
 	CollisionVisualizerComponent->SetHiddenInGame(true);
 	CollisionVisualizerComponent->SetCastShadow(false);
+	CollisionVisualizerComponent->SetCollisionEnabled(ECollisionEnabled::Type::NoCollision);
 
 	// Setup Mesh
 	if (UStaticMesh* CubeMesh = LoadObject<UStaticMesh>(nullptr, TEXT("/Script/Engine.StaticMesh'/Engine/BasicShapes/Cube.Cube'")))
diff --git a/Engine/Plugins/VirtualProduction/ClonerEffector/Source/ClonerEffector/Private/Subsystems/CEEffectorSubsystem.cpp b/Engine/Plugins/VirtualProduction/ClonerEffector/Source/ClonerEffector/Private/Subsystems/CEEffectorSubsystem.cpp
--- a/Engine/Plugins/VirtualProduction/ClonerEffector/Source/ClonerEffector/Private/Subsystems/CEEffectorSubsystem.cpp
+++ b/Engine/Plugins/VirtualProduction/ClonerEffector/Source/ClonerEffector/Private/Subsystems/CEEffectorSubsystem.cpp
@@ -2,6 +2,7 @@
 
 #include "Subsystems/CEEffectorSubsystem.h"
 
+#include "Algo/Transform.h"
 #include "Effector/CEEffectorComponent.h"
 #include "Effector/CEEffectorExtensionBase.h"
 #include "Effector/Logs/CEEffectorLogs.h"
@@ -66,6 +67,16 @@ void UCEEffectorSubsystem::Initialize(FSubsystemCollectionBase& Collection)
 	ScanForRegistrableClasses();
 
 	OnSubsystemInitializedDelegate.Broadcast();
+
+	FWorldDelegates::OnWorldTickStart.AddWeakLambda(this, [this](UWorld* InWorld, ELevelTick InLevel, float InDelta)
+	{
+		const bool bValidWorld = (InWorld->IsGameWorld() || InWorld->IsEditorWorld()) && !InWorld->IsPreviewWorld();
+
+		if (bValidWorld && InLevel != LEVELTICK_TimeOnly)
+		{
+			UpdateEffectorChannel(InWorld);
+		}
+	});
 }
 
 bool UCEEffectorSubsystem::RegisterChannelEffector(UCEEffectorComponent* InEffector)
@@ -324,9 +335,30 @@ void UCEEffectorSubsystem::UpdateEffectorChannel(const UWorld* InWorld)
 		return;
 	}
 
+	// Retrieve effectors in the input world
+	TArray<UCEEffectorComponent*> WorldEffectors;
+	Algo::TransformIf(
+		EffectorsWeak,
+		WorldEffectors,
+		[InWorld](const TWeakObjectPtr<UCEEffectorComponent>& InEffectorWeak)
+		{
+			const UCEEffectorComponent* Effector = InEffectorWeak.Get();
+			return IsValid(Effector) && Effector->GetWorld() == InWorld;
+		},
+		[](const TWeakObjectPtr<UCEEffectorComponent>& InEffectorWeak)
+		{
+			return InEffectorWeak.Get();
+		}
+	);
+
+	if (WorldEffectors.IsEmpty())
+	{
+		return;
+	}
+
 	// Reserve space in channel for each effectors
 	static const FNiagaraDataChannelSearchParameters SearchParameters;
-	UNiagaraDataChannelWriter* ChannelWriter = UNiagaraDataChannelLibrary::WriteToNiagaraDataChannel(InWorld, EffectorDataChannelAsset.Get(), SearchParameters, EffectorsWeak.Num(), true, true, true, UCEEffectorSubsystem::StaticClass()->GetName());
+	UNiagaraDataChannelWriter* ChannelWriter = UNiagaraDataChannelLibrary::WriteToNiagaraDataChannel(InWorld, EffectorDataChannelAsset.Get(), SearchParameters, WorldEffectors.Num(), true, true, true, UCEEffectorSubsystem::StaticClass()->GetName());
 
 	if (!ChannelWriter)
 	{
@@ -336,16 +368,8 @@ void UCEEffectorSubsystem::UpdateEffectorChannel(const UWorld* InWorld)
 
 	// Remove invalid effectors and push updates to effector assigned channel indexes
 	int32 EffectorIndex = 0;
-	for (const TWeakObjectPtr<UCEEffectorComponent>& EffectorWeak : EffectorsWeak)
+	for (UCEEffectorComponent* Effector : WorldEffectors)
 	{
-		UCEEffectorComponent* Effector = EffectorWeak.Get();
-
-		if (!IsValid(Effector)
-			|| Effector->GetWorld() != InWorld)
-		{
-			continue;
-		}
-
 		FCEClonerEffectorChannelData& ChannelData = Effector->GetChannelData();
 
 		const int32 PreviousIdentifier = ChannelData.Identifier;
@@ -383,35 +407,11 @@ TStatId UCEEffectorSubsystem::GetStatId() const
 
 void UCEEffectorSubsystem::Tick(float InDeltaTime)
 {
-	TSet<const UWorld*> Worlds;
-
-	for (TArray<TWeakObjectPtr<UCEEffectorComponent>>::TIterator It(EffectorsWeak); It; ++It)
-	{
-		const UCEEffectorComponent* Effector = It->Get();
-		if (!IsValid(Effector))
-		{
-			It.RemoveCurrent();
-			continue;
-		}
-
-		const UWorld* EffectorWorld = Effector->GetWorld();
-		if (!IsValid(EffectorWorld) || !EffectorWorld->IsInitialized() || EffectorWorld->IsBeingCleanedUp())
-		{
-			continue;
-		}
-
-		Worlds.Add(EffectorWorld);
-	}
-
-	for (const UWorld* World : Worlds)
-	{
-		UpdateEffectorChannel(World);
-	}
 }
 
 bool UCEEffectorSubsystem::IsTickable() const
 {
-	return !EffectorsWeak.IsEmpty();
+	return false;
 }
 
 #undef LOCTEXT_NAMESPACE
