MotionDesign : ClonerEffector
- Fixed rotation issue causing "jumps" once a property of the cloned object changed, this is due to multiple rotation being used on the cloned object, cloner and parent actor
- Fixed mismatch actor bounds warning would trigger even when cached and actual bounds didn't change
- Fixed create cloner outliner action will update actor mobility from static to movable to attach it under the cloner

diff --git a/Engine/Plugins/VirtualProduction/ClonerEffector/Source/ClonerEffector/Private/Cloner/Attachments/CEClonerAttachmentRootItem.cpp b/Engine/Plugins/VirtualProduction/ClonerEffector/Source/ClonerEffector/Private/Cloner/Attachments/CEClonerAttachmentRootItem.cpp
new file mode 100644
--- /dev/null
+++ b/Engine/Plugins/VirtualProduction/ClonerEffector/Source/ClonerEffector/Private/Cloner/Attachments/CEClonerAttachmentRootItem.cpp
@@ -0,0 +1,9 @@
+// Copyright Epic Games, Inc. All Rights Reserved.
+
+#include "Cloner/Attachments/CEClonerAttachmentRootItem.h"
+#include "GameFramework/Actor.h"
+
+FCEClonerAttachmentRootItem::FCEClonerAttachmentRootItem(AActor* InActor)
+	: RootActor(InActor)
+{
+}
diff --git a/Engine/Plugins/VirtualProduction/ClonerEffector/Source/ClonerEffector/Private/Cloner/Attachments/CEClonerAttachmentTree.cpp b/Engine/Plugins/VirtualProduction/ClonerEffector/Source/ClonerEffector/Private/Cloner/Attachments/CEClonerAttachmentTree.cpp
--- a/Engine/Plugins/VirtualProduction/ClonerEffector/Source/ClonerEffector/Private/Cloner/Attachments/CEClonerAttachmentTree.cpp
+++ b/Engine/Plugins/VirtualProduction/ClonerEffector/Source/ClonerEffector/Private/Cloner/Attachments/CEClonerAttachmentTree.cpp
@@ -54,8 +54,7 @@ void FCEClonerAttachmentTree::Cleanup()
 void FCEClonerAttachmentTree::Reset()
 {
 	ItemAttachmentMap.Reset();
-	RootActors.Reset();
-	MergedBakedMeshes.Reset();
+	RootItemAttachments.Reset();
 	DirtyItemAttachments.Reset();
 	Status = ECEClonerAttachmentStatus::Outdated;
 }
@@ -128,9 +127,9 @@ AActor* FCEClonerAttachmentTree::FindRootActor(AActor* InActor)
 TArray<AActor*> FCEClonerAttachmentTree::GetRootActors() const
 {
 	TArray<AActor*> Actors;
-	Algo::Transform(RootActors, Actors, [](const TWeakObjectPtr<AActor>& InActorWeak)->AActor*
+	Algo::Transform(RootItemAttachments, Actors, [](const FCEClonerAttachmentRootItem& InRootAttachment)->AActor*
 	{
-		return InActorWeak.Get();
+		return InRootAttachment.RootActor.Get();
 	});
 	return Actors;
 }
@@ -160,8 +159,10 @@ void FCEClonerAttachmentTree::DetachItem(AActor* InActor)
 bool FCEClonerAttachmentTree::IsCacheAvailable(bool bInAllowInvalid) const
 {
 	// Check cached meshes are not async loading
-	for (const TObjectPtr<UStaticMesh>& BakedMesh : MergedBakedMeshes)
+	for (const FCEClonerAttachmentRootItem& RootItem : RootItemAttachments)
 	{
+		const UStaticMesh* BakedMesh = RootItem.MergedBakedMesh.Get();
+
 		if (!IsValid(BakedMesh))
 		{
 			if (!bInAllowInvalid)
@@ -242,26 +243,22 @@ FBox FCEClonerAttachmentTree::GetAttachmentBounds(AActor* InActor, bool bInInclu
 	return Bounds;
 }
 
-FBox FCEClonerAttachmentTree::GetCacheBounds(AActor* InActor) const
+#if WITH_EDITOR
+TOptional<FVector> FCEClonerAttachmentTree::GetRootCachedMeshSize(AActor* InActor) const
 {
-	FBox Bounds(ForceInitToZero);
-	const int32 Index = RootActors.Find(InActor);
-
-	if (Index == INDEX_NONE)
-	{
-		return Bounds;
-	}
+	TOptional<FVector> MeshSize;
+	const FCEClonerAttachmentRootItem* RootItem = RootItemAttachments.FindByKey(InActor);
 
-	if (!MergedBakedMeshes.IsValidIndex(Index) || !MergedBakedMeshes[Index])
+	if (!RootItem || !RootItem->MergedBakedMesh.Get())
 	{
-		return Bounds;
+		return MeshSize;
 	}
 
-	Bounds = MergedBakedMeshes[Index]->GetBounds().GetBox();
-	Bounds.IsValid = true;
+	MeshSize = RootItem->MeshSize;
 
-	return Bounds;
+	return MeshSize;
 }
+#endif
 
 bool FCEClonerAttachmentTree::UpdateAttachmentsInternal()
 {
@@ -275,10 +272,8 @@ bool FCEClonerAttachmentTree::UpdateAttachmentsInternal()
 	TArray<AActor*> RootChildren;
 	BehaviorImplementation->GetOrderedChildrenActors(*this, AttachmentRoot, RootChildren);
 
-	TArray<TObjectPtr<UStaticMesh>> NewCombinesMeshes;
-	TArray<TWeakObjectPtr<AActor>> NewRootActors;
-	NewRootActors.Reserve(RootChildren.Num());
-	NewCombinesMeshes.Reserve(RootChildren.Num());
+	TArray<FCEClonerAttachmentRootItem> NewRootItems;
+	NewRootItems.Reserve(RootChildren.Num());
 	for (int32 RootIdx = 0; RootIdx < RootChildren.Num(); RootIdx++)
 	{
 		AActor* RootChild = RootChildren[RootIdx];
@@ -292,24 +287,27 @@ bool FCEClonerAttachmentTree::UpdateAttachmentsInternal()
 		UpdateAttachment(RootChild, nullptr, RootTransform);
 
 		// Lets find the old root idx
-		const int32 OldIdx = RootActors.Find(RootChild);
-		TObjectPtr<UStaticMesh> CombineBakedMesh = nullptr;
-		if (MergedBakedMeshes.IsValidIndex(OldIdx))
+		if (const FCEClonerAttachmentRootItem* RootItem = RootItemAttachments.FindByKey(RootChild))
 		{
-			CombineBakedMesh = MergedBakedMeshes[OldIdx].Get();
+			const int32 OldIdx = RootItemAttachments.Find(*RootItem);
 
 			// Did we rearrange stuff ?
 			if (RootIdx != OldIdx)
 			{
 				bItemAttachmentsDirty = true;
 			}
+
+			NewRootItems.Add(*RootItem);
+		}
+		else
+		{
+			FCEClonerAttachmentRootItem NewRootItem(RootChild);
+			NewRootItems.Add(NewRootItem);
 		}
-		NewCombinesMeshes.Add(CombineBakedMesh);
-		NewRootActors.Add(RootChild);
 	}
 
 	// Did we remove any root actors ?
-	if (RootActors.Num() != NewRootActors.Num())
+	if (RootItemAttachments.Num() != NewRootItems.Num())
 	{
 		bItemAttachmentsDirty = true;
 	}
@@ -356,8 +354,8 @@ bool FCEClonerAttachmentTree::UpdateAttachmentsInternal()
 		bItemAttachmentsDirty = true;
 	}
 
-	RootActors = NewRootActors;
-	MergedBakedMeshes = NewCombinesMeshes;
+	RootItemAttachments = NewRootItems;
+
 	return bItemAttachmentsDirty;
 }
 
diff --git a/Engine/Plugins/VirtualProduction/ClonerEffector/Source/ClonerEffector/Private/Cloner/Attachments/CEClonerAttachmentTreeBehavior.cpp b/Engine/Plugins/VirtualProduction/ClonerEffector/Source/ClonerEffector/Private/Cloner/Attachments/CEClonerAttachmentTreeBehavior.cpp
--- a/Engine/Plugins/VirtualProduction/ClonerEffector/Source/ClonerEffector/Private/Cloner/Attachments/CEClonerAttachmentTreeBehavior.cpp
+++ b/Engine/Plugins/VirtualProduction/ClonerEffector/Source/ClonerEffector/Private/Cloner/Attachments/CEClonerAttachmentTreeBehavior.cpp
@@ -144,10 +144,9 @@ void FCEClonerAttachmentGroupBehavior::InvalidateMesh(FCEClonerAttachmentTree& I
 	{
 		if (FoundItem->bRootItem || !FoundItem->ParentActor.IsValid())
 		{
-			const int32 RootIdx = InTree.RootActors.Find(InActor);
-			if (InTree.MergedBakedMeshes.IsValidIndex(RootIdx))
+			if (FCEClonerAttachmentRootItem* RootItem = InTree.RootItemAttachments.FindByKey(InActor))
 			{
-				InTree.MergedBakedMeshes[RootIdx] = nullptr;
+				RootItem->Reset();
 				InTree.bItemAttachmentsDirty = true;
 			}
 		}
@@ -169,24 +168,40 @@ void FCEClonerAttachmentGroupBehavior::OnItemAttached(FCEClonerAttachmentTree& I
 		// Ensure cached mesh bounds matches with attachment bounds
 		if (InItem.bRootItem)
 		{
-			AActor* ItemActor = InItem.ItemActor.Get();
+			const TOptional<FVector> BakedMeshSize = InTree.GetRootCachedMeshSize(Actor);
 
-			const FBox BakedMeshBounds = InTree.GetCacheBounds(ItemActor);
-			const FBox AttachmentBounds = InTree.GetAttachmentBounds(ItemActor, /** IncludeChildren */true);
+			// Refresh cloner attachment when mesh size is not set
+			if (!BakedMeshSize.IsSet())
+			{
+				InItem.MeshStatus = ECEClonerAttachmentStatus::Outdated;
+				InvalidateMesh(InTree, Actor);
+				InTree.DirtyItemAttachments.Add(Actor);
+				InTree.MarkAttachmentOutdated();
+				return;
+			}
 
-			if (BakedMeshBounds.IsValid
-				&& AttachmentBounds.IsValid
-				&& !BakedMeshBounds.GetSize().Equals(AttachmentBounds.GetSize(), 1.f))
+			// Leave the opportunity to dynamic mesh to initialize and build their geometry
+			FTSTicker::GetCoreTicker().AddTicker(FTickerDelegate::CreateSPLambda(this, [this, &InTree, BakedMeshSize, Actor](float)->bool
 			{
-				const AActor* ClonerActor = InTree.GetAttachmentRoot();
+				if (IsValid(Actor))
+				{
+					const FVector AttachmentSize = InTree.GetAttachmentBounds(Actor, /** IncludeChildren */true).GetSize();
 
-				UE::ClonerEffector::Utilities::ShowWarning(FText::Format(
-					LOCTEXT("AttachmentCachedBoundsMismatch", "Cloner {0} : {1} cached bounds vs actor bounds mismatch, see logs"),
-					FText::FromString(ClonerActor->GetActorNameOrLabel()),
-					FText::FromString(Actor->GetActorNameOrLabel())));
+					if (!BakedMeshSize->Equals(AttachmentSize, 1.f))
+					{
+						const AActor* ClonerActor = InTree.GetAttachmentRoot();
 
-				UE_LOG(LogCECloner, Warning, TEXT("%s : %s bounds mismatch, cached bounds (%s) vs actor bounds (%s), verify attached actor"), *ClonerActor->GetActorNameOrLabel(), *Actor->GetActorNameOrLabel(), *BakedMeshBounds.ToString(), *AttachmentBounds.ToString());
-			}
+						UE::ClonerEffector::Utilities::ShowWarning(FText::Format(
+							LOCTEXT("AttachmentCachedBoundsMismatch", "Cloner {0} : {1} cached size vs actor size mismatch, see logs"),
+							FText::FromString(ClonerActor->GetActorNameOrLabel()),
+							FText::FromString(Actor->GetActorNameOrLabel())));
+
+						UE_LOG(LogCECloner, Warning, TEXT("%s : %s bounds mismatch, cached size (%s) vs actor size (%s), verify attached actor"), *ClonerActor->GetActorNameOrLabel(), *Actor->GetActorNameOrLabel(), *BakedMeshSize->ToString(), *AttachmentSize.ToString());
+					}
+				}
+
+				return false;
+			}));
 		}
 #endif
 	}
diff --git a/Engine/Plugins/VirtualProduction/ClonerEffector/Source/ClonerEffector/Private/Cloner/CEClonerComponent.cpp b/Engine/Plugins/VirtualProduction/ClonerEffector/Source/ClonerEffector/Private/Cloner/CEClonerComponent.cpp
--- a/Engine/Plugins/VirtualProduction/ClonerEffector/Source/ClonerEffector/Private/Cloner/CEClonerComponent.cpp
+++ b/Engine/Plugins/VirtualProduction/ClonerEffector/Source/ClonerEffector/Private/Cloner/CEClonerComponent.cpp
@@ -17,6 +17,7 @@
 #include "Logs/CEClonerLogs.h"
 #include "Materials/Material.h"
 #include "Materials/MaterialInstanceDynamic.h"
+#include "Misc/EnumerateRange.h"
 #include "NiagaraMeshRendererProperties.h"
 #include "NiagaraSystem.h"
 #include "Settings/CEClonerEffectorSettings.h"
@@ -302,9 +303,9 @@ void UCEClonerComponent::UpdateDirtyMeshesAsync()
 			}
 
 			// Update actors baked static mesh
-			for (int32 Idx = 0; Idx < This->ClonerTree.RootActors.Num(); Idx++)
+			for (const FCEClonerAttachmentRootItem& RootItem : This->ClonerTree.RootItemAttachments)
 			{
-				const UStaticMesh* RootStaticMesh = This->ClonerTree.MergedBakedMeshes[Idx].Get();
+				const UStaticMesh* RootStaticMesh = RootItem.MergedBakedMesh.Get();
 
 				if (!RootStaticMesh)
 				{
@@ -314,7 +315,7 @@ void UCEClonerComponent::UpdateDirtyMeshesAsync()
 						break;
 					}
 
-					AActor* RootActor = This->ClonerTree.RootActors[Idx].Get();
+					AActor* RootActor = RootItem.RootActor.Get();
 					This->UpdateActorStaticMesh(RootActor);
 				}
 			}
@@ -393,16 +394,16 @@ void UCEClonerComponent::UpdateActorStaticMesh(AActor* InRootActor)
 		return;
 	}
 
-	const int32 RootIdx = ClonerTree.RootActors.Find(InRootActor);
+	FCEClonerAttachmentRootItem* RootItem = ClonerTree.RootItemAttachments.FindByKey(InRootActor);
 
-	if (RootIdx == INDEX_NONE)
+	if (!RootItem)
 	{
 		return;
 	}
 
-	const FCEClonerAttachmentItem* RootAttachmentItem = ClonerTree.ItemAttachmentMap.Find(InRootActor);
+	const FCEClonerAttachmentItem* ActorAttachmentItem = ClonerTree.ItemAttachmentMap.Find(InRootActor);
 
-	if (!RootAttachmentItem)
+	if (!ActorAttachmentItem)
 	{
 		return;
 	}
@@ -440,7 +441,10 @@ void UCEClonerComponent::UpdateActorStaticMesh(AActor* InRootActor)
 	ClonerTree.bItemAttachmentsDirty = MeshBuilder.BuildStaticMesh(Mesh, MeshMaterials);
 	MeshBuilder.Reset();
 
-	ClonerTree.MergedBakedMeshes[RootIdx] = Mesh;
+	RootItem->MergedBakedMesh = Mesh;
+#if WITH_EDITORONLY_DATA
+	RootItem->MeshSize = ClonerTree.GetAttachmentBounds(InRootActor, /** IncludeChildren */true).GetSize();
+#endif
 
 	UE_LOG(LogCECloner, Log, TEXT("%s : Updated actor static mesh : %s - %i - %i"), *ClonerActor->GetActorNameOrLabel(), *InRootActor->GetActorNameOrLabel(), Mesh->GetNumTriangles(/** LOD */0), AttachmentItems.Num());
 }
@@ -471,9 +475,9 @@ void UCEClonerComponent::GetPreloadDependencies(TArray<UObject*>& OutDeps)
 {
 	Super::GetPreloadDependencies(OutDeps);
 
-	for (UStaticMesh* BakedMesh : ClonerTree.MergedBakedMeshes)
+	for (const FCEClonerAttachmentRootItem& RootItem : ClonerTree.RootItemAttachments)
 	{
-		if (BakedMesh)
+		if (UStaticMesh* BakedMesh = RootItem.MergedBakedMesh.Get())
 		{
 			OutDeps.Add(BakedMesh);
 		}
@@ -1445,22 +1449,23 @@ void UCEClonerComponent::UpdateClonerMeshes()
 		return;
 	}
 
-	bool bMeshChanged = MeshRenderer->Meshes.Num() != ClonerTree.MergedBakedMeshes.Num();
+	bool bMeshChanged = MeshRenderer->Meshes.Num() != ClonerTree.RootItemAttachments.Num();
 
 	if (ClonerTree.bItemAttachmentsDirty)
 	{
 		// Resize mesh array properly
-		if (MeshRenderer->Meshes.Num() > ClonerTree.MergedBakedMeshes.Num())
+		if (MeshRenderer->Meshes.Num() > ClonerTree.RootItemAttachments.Num())
 		{
-			MeshRenderer->Meshes.SetNum(ClonerTree.MergedBakedMeshes.Num());
+			MeshRenderer->Meshes.SetNum(ClonerTree.RootItemAttachments.Num());
 		}
 
 		// Set baked meshes in mesh renderer array
-		for (int32 Idx = 0; Idx < ClonerTree.MergedBakedMeshes.Num(); Idx++)
+		for (TConstEnumerateRef<FCEClonerAttachmentRootItem> RootItem : EnumerateRange(ClonerTree.RootItemAttachments))
 		{
-			UStaticMesh* StaticMesh = ClonerTree.MergedBakedMeshes[Idx].Get();
+			UStaticMesh* StaticMesh = RootItem->MergedBakedMesh.Get();
+			const int32 Index = RootItem.GetIndex();
 
-			FNiagaraMeshRendererMeshProperties& MeshProperties = !MeshRenderer->Meshes.IsValidIndex(Idx) ? MeshRenderer->Meshes.AddDefaulted_GetRef() : MeshRenderer->Meshes[Idx];
+			FNiagaraMeshRendererMeshProperties& MeshProperties = !MeshRenderer->Meshes.IsValidIndex(Index) ? MeshRenderer->Meshes.AddDefaulted_GetRef() : MeshRenderer->Meshes[Index];
 			bMeshChanged |= MeshProperties.Mesh != StaticMesh;
 
 			if (StaticMesh && StaticMesh->GetNumTriangles(0) > 0)
@@ -1472,12 +1477,12 @@ void UCEClonerComponent::UpdateClonerMeshes()
 				MeshProperties.Mesh = nullptr;
 			}
 
-			if (ClonerTree.RootActors.IsValidIndex(Idx))
+			if (const AActor* RootActor = RootItem->RootActor.Get())
 			{
-				if (const AActor* RootActor = ClonerTree.RootActors[Idx].Get())
+				if (const USceneComponent* SceneComponent = RootActor->GetRootComponent())
 				{
-					MeshProperties.Rotation = GlobalRotation + RootActor->GetActorRotation();
-					MeshProperties.Scale = GlobalScale * RootActor->GetActorScale3D();
+					MeshProperties.Rotation = GlobalRotation + SceneComponent->GetRelativeRotation();
+					MeshProperties.Scale = GlobalScale * SceneComponent->GetRelativeScale3D();
 				}
 			}
 		}
@@ -1491,11 +1496,11 @@ void UCEClonerComponent::UpdateClonerMeshes()
 	}
 
 	// Extensions could override mesh renderer meshes array
-	bMeshChanged |= ClonerTree.MergedBakedMeshes.Num() != MeshRenderer->Meshes.Num();
+	bMeshChanged |= ClonerTree.RootItemAttachments.Num() != MeshRenderer->Meshes.Num();
 
 	if (bMeshChanged)
 	{
-		UE_LOG(LogCECloner, Log, TEXT("%s : Cloner mesh updated - %i cached meshes - %i rendered meshes"), *ClonerActor->GetActorNameOrLabel(), ClonerTree.MergedBakedMeshes.Num(), MeshRenderer->Meshes.Num())
+		UE_LOG(LogCECloner, Log, TEXT("%s : Cloner mesh updated - %i cached meshes - %i rendered meshes"), *ClonerActor->GetActorNameOrLabel(), ClonerTree.RootItemAttachments.Num(), MeshRenderer->Meshes.Num())
 	}
 
 	// Set new number of meshes in renderer
diff --git a/Engine/Plugins/VirtualProduction/ClonerEffector/Source/ClonerEffector/Private/Subsystems/CEClonerSubsystem.cpp b/Engine/Plugins/VirtualProduction/ClonerEffector/Source/ClonerEffector/Private/Subsystems/CEClonerSubsystem.cpp
--- a/Engine/Plugins/VirtualProduction/ClonerEffector/Source/ClonerEffector/Private/Subsystems/CEClonerSubsystem.cpp
+++ b/Engine/Plugins/VirtualProduction/ClonerEffector/Source/ClonerEffector/Private/Subsystems/CEClonerSubsystem.cpp
@@ -504,6 +504,11 @@ UCEClonerComponent* UCEClonerSubsystem::CreateClonerWithActors(UWorld* InWorld,
 					Actor->Modify();
 #endif
 
+					if (USceneComponent* RootComponent = Actor->GetRootComponent())
+					{
+						RootComponent->SetMobility(EComponentMobility::Type::Movable);
+					}
+
 					Actor->AttachToActor(NewClonerActor, FAttachmentTransformRules::KeepWorldTransform);
 				}
 			}
diff --git a/Engine/Plugins/VirtualProduction/ClonerEffector/Source/ClonerEffector/Private/Subsystems/CEEffectorSubsystem.cpp b/Engine/Plugins/VirtualProduction/ClonerEffector/Source/ClonerEffector/Private/Subsystems/CEEffectorSubsystem.cpp
--- a/Engine/Plugins/VirtualProduction/ClonerEffector/Source/ClonerEffector/Private/Subsystems/CEEffectorSubsystem.cpp
+++ b/Engine/Plugins/VirtualProduction/ClonerEffector/Source/ClonerEffector/Private/Subsystems/CEEffectorSubsystem.cpp
@@ -79,6 +79,13 @@ void UCEEffectorSubsystem::Initialize(FSubsystemCollectionBase& Collection)
 	});
 }
 
+void UCEEffectorSubsystem::Deinitialize()
+{
+	Super::Deinitialize();
+
+	FWorldDelegates::OnWorldTickStart.RemoveAll(this);
+}
+
 bool UCEEffectorSubsystem::RegisterChannelEffector(UCEEffectorComponent* InEffector)
 {
 	if (!IsValid(InEffector) || !InEffector->GetOwner())
diff --git a/Engine/Plugins/VirtualProduction/ClonerEffector/Source/ClonerEffector/Public/Cloner/Attachments/CEClonerAttachmentRootItem.h b/Engine/Plugins/VirtualProduction/ClonerEffector/Source/ClonerEffector/Public/Cloner/Attachments/CEClonerAttachmentRootItem.h
new file mode 100644
--- /dev/null
+++ b/Engine/Plugins/VirtualProduction/ClonerEffector/Source/ClonerEffector/Public/Cloner/Attachments/CEClonerAttachmentRootItem.h
@@ -0,0 +1,59 @@
+// Copyright Epic Games, Inc. All Rights Reserved.
+
+#pragma once
+
+#include "Containers/Array.h"
+#include "CEClonerAttachmentTypes.h"
+#include "Misc/Optional.h"
+#include "UObject/ObjectPtr.h"
+#include "UObject/WeakObjectPtrTemplates.h"
+#include "CEClonerAttachmentRootItem.generated.h"
+
+class AActor;
+class UStaticMesh;
+
+USTRUCT()
+struct FCEClonerAttachmentRootItem
+{
+	GENERATED_BODY()
+
+	FCEClonerAttachmentRootItem()
+	{}
+
+	FCEClonerAttachmentRootItem(AActor* InActor);
+
+	friend uint32 GetTypeHash(const FCEClonerAttachmentRootItem& InItem)
+	{
+		return GetTypeHash(InItem.RootActor);
+	}
+
+	bool operator==(const FCEClonerAttachmentRootItem& InOther) const
+	{
+		return RootActor.IsValid() && RootActor.Get() == InOther.RootActor.Get();
+	}
+
+	void Reset()
+	{
+		MergedBakedMesh = nullptr;
+#if WITH_EDITORONLY_DATA
+		MeshSize.Reset();
+#endif
+	}
+
+	/** Actors directly attached to the cloner actor, order is important here */
+	UPROPERTY()
+	TWeakObjectPtr<AActor> RootActor;
+
+	/**
+	 * Merged static meshes corresponding to the root actors for niagara,
+	 * result of merging dynamic meshes together
+	 */
+	UPROPERTY()
+	TObjectPtr<UStaticMesh> MergedBakedMesh;
+
+#if WITH_EDITORONLY_DATA
+	/** Mesh size of the baked mesh / root actor tree, used to compare and detect a change */
+	UPROPERTY()
+	TOptional<FVector> MeshSize;
+#endif
+};
\ No newline at end of file
diff --git a/Engine/Plugins/VirtualProduction/ClonerEffector/Source/ClonerEffector/Public/Cloner/Attachments/CEClonerAttachmentTree.h b/Engine/Plugins/VirtualProduction/ClonerEffector/Source/ClonerEffector/Public/Cloner/Attachments/CEClonerAttachmentTree.h
--- a/Engine/Plugins/VirtualProduction/ClonerEffector/Source/ClonerEffector/Public/Cloner/Attachments/CEClonerAttachmentTree.h
+++ b/Engine/Plugins/VirtualProduction/ClonerEffector/Source/ClonerEffector/Public/Cloner/Attachments/CEClonerAttachmentTree.h
@@ -3,6 +3,7 @@
 #pragma once
 
 #include "CEClonerAttachmentItem.h"
+#include "CEClonerAttachmentRootItem.h"
 #include "CEClonerAttachmentTypes.h"
 #include "Components/SceneComponent.h"
 #include "Containers/Array.h"
@@ -83,23 +84,18 @@ struct FCEClonerAttachmentTree
 	/** Get the attachment bounds based on tree hierarchy */
 	FBox GetAttachmentBounds(AActor* InActor, bool bInIncludeChildren) const;
 
-	/** Get cached bounds for a root actor */
-	FBox GetCacheBounds(AActor* InActor) const;
+#if WITH_EDITOR
+	/** Get cached size of root mesh, -1,-1,-1 when invalid */
+	TOptional<FVector> GetRootCachedMeshSize(AActor* InActor) const;
+#endif
 
 	/** All cloner attached actor items */
 	UPROPERTY()
 	TMap<TWeakObjectPtr<AActor>, FCEClonerAttachmentItem> ItemAttachmentMap;
 
-	/** Actors directly attached to the cloner actor, order is important here */
+	/** Root items directly attached to the cloner actor, order is important here */
 	UPROPERTY()
-	TArray<TWeakObjectPtr<AActor>> RootActors;
-
-	/**
-	 * Merged static meshes corresponding to the root actors for niagara,
-	 * result of merging dynamic meshes together
-	 */
-	UPROPERTY()
-	TArray<TObjectPtr<UStaticMesh>> MergedBakedMeshes;
+	TArray<FCEClonerAttachmentRootItem> RootItemAttachments;
 
 	/** Status of the cloner tree */
 	UPROPERTY(Transient)
@@ -127,6 +123,7 @@ private:
 	/** Called when tree item is detached */
 	FOnAttachmentChanged OnItemDetachedDelegate;
 
+	UPROPERTY(Transient)
 	TObjectPtr<AActor> AttachmentRoot;
 
 	TSharedPtr<ICEClonerAttachmentTreeBehavior> BehaviorImplementation;
diff --git a/Engine/Plugins/VirtualProduction/ClonerEffector/Source/ClonerEffector/Public/Subsystems/CEEffectorSubsystem.h b/Engine/Plugins/VirtualProduction/ClonerEffector/Source/ClonerEffector/Public/Subsystems/CEEffectorSubsystem.h
--- a/Engine/Plugins/VirtualProduction/ClonerEffector/Source/ClonerEffector/Public/Subsystems/CEEffectorSubsystem.h
+++ b/Engine/Plugins/VirtualProduction/ClonerEffector/Source/ClonerEffector/Public/Subsystems/CEEffectorSubsystem.h
@@ -96,6 +96,7 @@ protected:
 
 	//~ Begin USubsystem
 	virtual void Initialize(FSubsystemCollectionBase& Collection) override;
+	virtual void Deinitialize() override;
 	//~ End USubsystem
 
 	//~ Begin FTickableGameObject
