MotionDesign : ClonerEffector
- Fixed mask read modifier would not apply itself on a cloner
- Fixed mask read material would be lost after level reload
- Fixed mask read material would not trigger a cloner update and require re-attaching actor to refresh state
- Enable niagara flag on SVG materials to be cloner compatible

#virtualized

diff --git a/Engine/Plugins/VirtualProduction/Avalanche/Source/AvalancheMask/Private/AvaMaskUtilities.cpp b/Engine/Plugins/VirtualProduction/Avalanche/Source/AvalancheMask/Private/AvaMaskUtilities.cpp
--- a/Engine/Plugins/VirtualProduction/Avalanche/Source/AvalancheMask/Private/AvaMaskUtilities.cpp
+++ b/Engine/Plugins/VirtualProduction/Avalanche/Source/AvalancheMask/Private/AvaMaskUtilities.cpp
@@ -98,8 +98,11 @@ namespace UE::AvaMask::Internal
 			// Rerun construction scripts
 			InActor->RerunConstructionScripts();
 #endif
+
+			// Broadcast component creation to let other systems know about it
+			USceneComponent::MarkRenderStateDirtyEvent.Broadcast(*Component);
 		}
-		
+
 		return Component;
 	}
 
diff --git a/Engine/Plugins/VirtualProduction/Avalanche/Source/AvalancheMask/Private/Mask2D/AvaMask2DBaseModifier.cpp b/Engine/Plugins/VirtualProduction/Avalanche/Source/AvalancheMask/Private/Mask2D/AvaMask2DBaseModifier.cpp
--- a/Engine/Plugins/VirtualProduction/Avalanche/Source/AvalancheMask/Private/Mask2D/AvaMask2DBaseModifier.cpp
+++ b/Engine/Plugins/VirtualProduction/Avalanche/Source/AvalancheMask/Private/Mask2D/AvaMask2DBaseModifier.cpp
@@ -5,6 +5,7 @@
 #include "AvaMaskLog.h"
 #include "AvaMaskSubsystem.h"
 #include "AvaMaskUtilities.h"
+#include "Components/BillboardComponent.h"
 #include "Components/PrimitiveComponent.h"
 #include "Engine/CanvasRenderTarget2D.h"
 #include "Engine/Engine.h"
@@ -23,6 +24,7 @@
 #include "Materials/Material.h"
 #include "Materials/MaterialInstanceDynamic.h"
 #include "Modifiers/ActorModifierCoreStack.h"
+#include "Particles/ParticleSystemComponent.h"
 #include "Subsystems/ActorModifierCoreSubsystem.h"
 #include "Text3DComponent.h"
 
@@ -599,8 +601,27 @@ bool UAvaMask2DBaseModifier::ActorSupportsMaskReadWrite(const AActor* InActor)
 {
 	// Only supports primitives and Text3D component.
 	// Note: Actors holding Text3D Components do eventually get Text Mesh Subobjects; but this should still return true even if these components are not yet built (the build is queued)
-	return InActor->FindComponentByClass<UText3DComponent>()
-		|| InActor->FindComponentByClass<UPrimitiveComponent>();
+	if (InActor->FindComponentByClass<UText3DComponent>())
+	{
+		return true;
+	}
+
+	TArray<UPrimitiveComponent*> PrimitiveComponents;
+	InActor->GetComponents(PrimitiveComponents);
+
+	// Skip components that are editor only (visualizations) and FX/Billboard
+	for (const UPrimitiveComponent* PrimitiveComponent : PrimitiveComponents)
+	{
+		if (PrimitiveComponent
+			&& !PrimitiveComponent->IsEditorOnly()
+			&& !PrimitiveComponent->IsA<UFXSystemComponent>()
+			&& !PrimitiveComponent->IsA<UBillboardComponent>())
+		{
+			return true;
+		}
+	}
+
+	return false;
 }
 
 bool UAvaMask2DBaseModifier::TryResolveParentChannel()
diff --git a/Engine/Plugins/VirtualProduction/ClonerEffector/Source/ClonerEffector/Private/Cloner/Attachments/CEClonerAttachmentTree.cpp b/Engine/Plugins/VirtualProduction/ClonerEffector/Source/ClonerEffector/Private/Cloner/Attachments/CEClonerAttachmentTree.cpp
--- a/Engine/Plugins/VirtualProduction/ClonerEffector/Source/ClonerEffector/Private/Cloner/Attachments/CEClonerAttachmentTree.cpp
+++ b/Engine/Plugins/VirtualProduction/ClonerEffector/Source/ClonerEffector/Private/Cloner/Attachments/CEClonerAttachmentTree.cpp
@@ -124,6 +124,16 @@ AActor* FCEClonerAttachmentTree::FindRootActor(AActor* InActor)
 	return nullptr;
 }
 
+FCEClonerAttachmentRootItem* FCEClonerAttachmentTree::FindRootItem(AActor* InActor)
+{
+	if (AActor* RootActor = FindRootActor(InActor))
+	{
+		return RootItemAttachments.FindByKey(RootActor);
+	}
+
+	return nullptr;
+}
+
 TArray<AActor*> FCEClonerAttachmentTree::GetRootActors() const
 {
 	TArray<AActor*> Actors;
diff --git a/Engine/Plugins/VirtualProduction/ClonerEffector/Source/ClonerEffector/Private/Cloner/Attachments/CEClonerAttachmentTreeBehavior.cpp b/Engine/Plugins/VirtualProduction/ClonerEffector/Source/ClonerEffector/Private/Cloner/Attachments/CEClonerAttachmentTreeBehavior.cpp
--- a/Engine/Plugins/VirtualProduction/ClonerEffector/Source/ClonerEffector/Private/Cloner/Attachments/CEClonerAttachmentTreeBehavior.cpp
+++ b/Engine/Plugins/VirtualProduction/ClonerEffector/Source/ClonerEffector/Private/Cloner/Attachments/CEClonerAttachmentTreeBehavior.cpp
@@ -12,6 +12,7 @@
 #include "Engine/World.h"
 #include "GameFramework/Actor.h"
 #include "Materials/Material.h"
+#include "Misc/EnumerateRange.h"
 #include "Subsystems/CEClonerSubsystem.h"
 #include "Utilities/CEClonerEffectorUtilities.h"
 
@@ -32,12 +33,16 @@ void FCEClonerAttachmentGroupBehavior::OnActivation(FCEClonerAttachmentTree& InT
 		}
 	}
 
-#if WITH_EDITOR
 	USceneComponent::MarkRenderStateDirtyEvent.AddSPLambda(this, [this, &InTree](UActorComponent& InComponent)
 	{
+		OnCheckMaterialChanged(InTree, InComponent.GetOwner());
+
+#if WITH_EDITOR
 		OnRenderStateDirty(InTree, InComponent);
+#endif
 	});
 
+#if WITH_EDITOR
 	FCoreUObjectDelegates::OnObjectPropertyChanged.AddSPLambda(this, [this, &InTree](UObject* InObject, FPropertyChangedEvent& InEvent)
 	{
 		OnCheckMaterialChanged(InTree, InObject);
@@ -87,9 +92,9 @@ void FCEClonerAttachmentGroupBehavior::OnDeactivation(FCEClonerAttachmentTree& I
 		}
 	}
 
-#if WITH_EDITOR
 	USceneComponent::MarkRenderStateDirtyEvent.RemoveAll(this);
 
+#if WITH_EDITOR
 	FCoreUObjectDelegates::OnObjectPropertyChanged.RemoveAll(this);
 
 	UMaterial::OnMaterialCompilationFinished().RemoveAll(this);
@@ -490,6 +495,20 @@ void FCEClonerAttachmentGroupBehavior::OnRenderStateDirty(FCEClonerAttachmentTre
 	InTree.MarkAttachmentOutdated();
 }
 
+void FCEClonerAttachmentGroupBehavior::OnLevelHierarchyChanged(FCEClonerAttachmentTree& InTree, AActor* InActor)
+{
+	if (IsValid(InActor))
+	{
+		const AActor* AttachmentRoot = InTree.GetAttachmentRoot();
+
+		if (InTree.ItemAttachmentMap.Contains(InActor) || InActor->IsAttachedTo(AttachmentRoot))
+		{
+			InTree.MarkAttachmentOutdated();
+		}
+	}
+}
+#endif
+
 void FCEClonerAttachmentGroupBehavior::OnCheckMaterialChanged(FCEClonerAttachmentTree& InTree, UObject* InObject)
 {
 	if (!IsValid(InObject))
@@ -532,23 +551,62 @@ void FCEClonerAttachmentGroupBehavior::OnCheckMaterialChanged(FCEClonerAttachmen
 		}
 
 		UE_LOG(LogCECloner, Log, TEXT("%s : Detected material change for %s"), *ClonerActor->GetActorNameOrLabel(), *ActorChanged->GetActorNameOrLabel());
-		InvalidateMesh(InTree, ActorChanged);
-		InTree.MarkAttachmentOutdated();
+
+		if (AttachmentItem->MeshStatus == ECEClonerAttachmentStatus::Outdated)
+		{
+			InvalidateMesh(InTree, ActorChanged);
+			InTree.MarkAttachmentOutdated();
+		}
+		else
+		{
+			RefreshBakedMeshRootMaterials(InTree, ActorChanged);
+		}
 	}
 }
 
-void FCEClonerAttachmentGroupBehavior::OnLevelHierarchyChanged(FCEClonerAttachmentTree& InTree, AActor* InActor)
+void FCEClonerAttachmentGroupBehavior::RefreshBakedMeshRootMaterials(FCEClonerAttachmentTree& InTree, AActor* InActor)
 {
-	if (IsValid(InActor))
+	FCEClonerAttachmentRootItem* RootItem = InTree.FindRootItem(InActor);
+	if (!RootItem || !RootItem->MergedBakedMesh)
 	{
-		const AActor* AttachmentRoot = InTree.GetAttachmentRoot();
+		return;
+	}
 
-		if (InTree.ItemAttachmentMap.Contains(InActor) || InActor->IsAttachedTo(AttachmentRoot))
+	TArray<const FCEClonerAttachmentItem*> AttachmentItems;
+	InTree.GetAttachments(RootItem->RootActor.Get(), AttachmentItems, /** Recurse */true);
+	if (AttachmentItems.IsEmpty())
+	{
+		return;
+	}
+
+	TArray<TWeakObjectPtr<UMaterialInterface>> RootMaterials;
+	TArray<FStaticMaterial>& StaticMaterials = RootItem->MergedBakedMesh->GetStaticMaterials();
+
+	int32 MaterialIndexOffset = 0;
+	for (const FCEClonerAttachmentItem* AttachmentItem : AttachmentItems)
+	{
+		if (!AttachmentItem)
 		{
-			InTree.MarkAttachmentOutdated();
+			continue;
+		}
+
+		for (TConstEnumerateRef<TWeakObjectPtr<UMaterialInterface>> BakedMaterialWeak : EnumerateRange(AttachmentItem->BakedMaterials))
+		{
+			if (UMaterialInterface* Material = BakedMaterialWeak->Get())
+			{
+				const int32 MaterialSlotIndex = MaterialIndexOffset + BakedMaterialWeak.GetIndex();
+
+				if (StaticMaterials.IsValidIndex(MaterialSlotIndex))
+				{
+					StaticMaterials[MaterialSlotIndex].MaterialInterface = Material;
+				}
+			}
+
+			MaterialIndexOffset++;
 		}
 	}
+
+	InTree.bItemAttachmentsDirty = true;
 }
-#endif
 
 #undef LOCTEXT_NAMESPACE
diff --git a/Engine/Plugins/VirtualProduction/ClonerEffector/Source/ClonerEffector/Public/Cloner/Attachments/CEClonerAttachmentItem.h b/Engine/Plugins/VirtualProduction/ClonerEffector/Source/ClonerEffector/Public/Cloner/Attachments/CEClonerAttachmentItem.h
--- a/Engine/Plugins/VirtualProduction/ClonerEffector/Source/ClonerEffector/Public/Cloner/Attachments/CEClonerAttachmentItem.h
+++ b/Engine/Plugins/VirtualProduction/ClonerEffector/Source/ClonerEffector/Public/Cloner/Attachments/CEClonerAttachmentItem.h
@@ -34,7 +34,7 @@ struct FCEClonerAttachmentItem
 	TObjectPtr<UDynamicMesh> BakedMesh;
 
 	/** Actual baked materials for the current actor item */
-	UPROPERTY(Transient)
+	UPROPERTY()
 	TArray<TWeakObjectPtr<UMaterialInterface>> BakedMaterials;
 
 	/** Status of the baked mesh, does it needs to be updated */
@@ -86,8 +86,8 @@ struct FCEClonerAttachmentItem
 
 	/** Get local attachment bounds */
 	FBox GetAttachmentBounds() const;
+#endif
 
 	/** Check this item materials to ensure they are supported */
 	bool CheckMaterialsChanged(bool bInUpdate, TArray<TWeakObjectPtr<UMaterialInterface>>* OutInvalidMaterials = nullptr);
-#endif
 };
\ No newline at end of file
diff --git a/Engine/Plugins/VirtualProduction/ClonerEffector/Source/ClonerEffector/Public/Cloner/Attachments/CEClonerAttachmentTree.h b/Engine/Plugins/VirtualProduction/ClonerEffector/Source/ClonerEffector/Public/Cloner/Attachments/CEClonerAttachmentTree.h
--- a/Engine/Plugins/VirtualProduction/ClonerEffector/Source/ClonerEffector/Public/Cloner/Attachments/CEClonerAttachmentTree.h
+++ b/Engine/Plugins/VirtualProduction/ClonerEffector/Source/ClonerEffector/Public/Cloner/Attachments/CEClonerAttachmentTree.h
@@ -66,6 +66,9 @@ struct FCEClonerAttachmentTree
 	/** Finds root actor of an actor based on this tree hierarchy by climbing up */
 	AActor* FindRootActor(AActor* InActor);
 
+	/** Finds root item of an actor based on this tree hierarchy by climbing up */
+	FCEClonerAttachmentRootItem* FindRootItem(AActor* InActor);
+
 	/** Retrieves all root actors from this tree hierarchy */
 	TArray<AActor*> GetRootActors() const;
 
diff --git a/Engine/Plugins/VirtualProduction/ClonerEffector/Source/ClonerEffector/Public/Cloner/Attachments/CEClonerAttachmentTreeBehavior.h b/Engine/Plugins/VirtualProduction/ClonerEffector/Source/ClonerEffector/Public/Cloner/Attachments/CEClonerAttachmentTreeBehavior.h
--- a/Engine/Plugins/VirtualProduction/ClonerEffector/Source/ClonerEffector/Public/Cloner/Attachments/CEClonerAttachmentTreeBehavior.h
+++ b/Engine/Plugins/VirtualProduction/ClonerEffector/Source/ClonerEffector/Public/Cloner/Attachments/CEClonerAttachmentTreeBehavior.h
@@ -65,13 +65,16 @@ protected:
 	/** Called when the render state of a cloned component changes */
 	void OnRenderStateDirty(FCEClonerAttachmentTree& InTree, UActorComponent& InComponent);
 
-	/** Called when material is done compiling and when property is changed to test whether material have changed */
-	void OnCheckMaterialChanged(FCEClonerAttachmentTree& InTree, UObject* InObject);
-
 	/** Called when a level actor hierarchy changes */
 	void OnLevelHierarchyChanged(FCEClonerAttachmentTree& InTree, AActor* InActor);
 #endif
 
+	/** Called when material is done compiling and when property is changed to test whether material have changed */
+	void OnCheckMaterialChanged(FCEClonerAttachmentTree& InTree, UObject* InObject);
+
+	/** Called to update all materials of a root actor baked mesh */
+	void RefreshBakedMeshRootMaterials(FCEClonerAttachmentTree& InTree, AActor* InActor);
+
 	FDelegateHandle WorldActorDestroyedDelegate;
 };
 
