MotionDesign : ClonerEffector
- Updated cloner default material to directly support geometry mask, default material is now masked instead of opaque

#virtualized

diff --git a/Engine/Plugins/VirtualProduction/ClonerEffector/ClonerEffector.uplugin b/Engine/Plugins/VirtualProduction/ClonerEffector/ClonerEffector.uplugin
--- a/Engine/Plugins/VirtualProduction/ClonerEffector/ClonerEffector.uplugin
+++ b/Engine/Plugins/VirtualProduction/ClonerEffector/ClonerEffector.uplugin
@@ -31,6 +31,10 @@
 		}
 	],
 	"Plugins": [
+		{
+			"Name": "GeometryMask",
+			"Enabled": true
+		},
 		{
 			"Name": "GeometryScripting",
 			"Enabled": true
