MotionDesign : ClonerEffector
- Fixed cloner crash when entering PIE
- Fixed 'Undo' action after nesting a shape in a Cloner Actor fails to return the cloner to its previous state
- Allow cloner and effector component to be blueprint spawnable

diff --git a/Engine/Plugins/VirtualProduction/ClonerEffector/Source/ClonerEffector/Private/Cloner/Attachments/CEClonerAttachmentTreeBehavior.cpp b/Engine/Plugins/VirtualProduction/ClonerEffector/Source/ClonerEffector/Private/Cloner/Attachments/CEClonerAttachmentTreeBehavior.cpp
--- a/Engine/Plugins/VirtualProduction/ClonerEffector/Source/ClonerEffector/Private/Cloner/Attachments/CEClonerAttachmentTreeBehavior.cpp
+++ b/Engine/Plugins/VirtualProduction/ClonerEffector/Source/ClonerEffector/Private/Cloner/Attachments/CEClonerAttachmentTreeBehavior.cpp
@@ -471,7 +471,7 @@ void FCEClonerAttachmentGroupBehavior::OnRenderStateDirty(FCEClonerAttachmentTre
 	AActor* Owner = InComponent.GetOwner();
 	const AActor* ClonerActor = InTree.GetAttachmentRoot();
 
-	if (!Owner || Owner->GetLevel() != ClonerActor->GetLevel())
+	if (!Owner || !ClonerActor || Owner->GetLevel() != ClonerActor->GetLevel())
 	{
 		return;
 	}
diff --git a/Engine/Plugins/VirtualProduction/ClonerEffector/Source/ClonerEffector/Private/Cloner/CEClonerComponent.cpp b/Engine/Plugins/VirtualProduction/ClonerEffector/Source/ClonerEffector/Private/Cloner/CEClonerComponent.cpp
--- a/Engine/Plugins/VirtualProduction/ClonerEffector/Source/ClonerEffector/Private/Cloner/CEClonerComponent.cpp
+++ b/Engine/Plugins/VirtualProduction/ClonerEffector/Source/ClonerEffector/Private/Cloner/CEClonerComponent.cpp
@@ -156,7 +156,7 @@ void UCEClonerComponent::PostEditUndo()
 	// Reregister ticker in case this object was destroyed then undo
 	RegisterTicker();
 
-	ForceUpdateCloner();
+	ClonerTree.MarkAttachmentOutdated();
 }
 
 const TCEPropertyChangeDispatcher<UCEClonerComponent> UCEClonerComponent::PropertyChangeDispatcher =
@@ -1113,18 +1113,21 @@ void UCEClonerComponent::OnLayoutNameChanged()
 #if WITH_EDITOR
 void UCEClonerComponent::OnVisualizerSpriteVisibleChanged()
 {
-	if (UTexture2D* SpriteTexture = LoadObject<UTexture2D>(nullptr, SpriteTexturePath))
+	if (GetWorld() != nullptr)
 	{
-		CreateSpriteComponent(SpriteTexture);
-
-		if (SpriteComponent)
+		if (UTexture2D* SpriteTexture = LoadObject<UTexture2D>(nullptr, SpriteTexturePath))
 		{
-			if (SpriteComponent->Sprite != SpriteTexture)
+			CreateSpriteComponent(SpriteTexture);
+
+			if (SpriteComponent)
 			{
-				SpriteComponent->SetSprite(SpriteTexture);
-			}
+				if (SpriteComponent->Sprite != SpriteTexture)
+				{
+					SpriteComponent->SetSprite(SpriteTexture);
+				}
 
-			SpriteComponent->SetVisibility(bVisualizerSpriteVisible, false);
+				SpriteComponent->SetVisibility(bVisualizerSpriteVisible, /** PropagateToChildren */false);
+			}
 		}
 	}
 }
diff --git a/Engine/Plugins/VirtualProduction/ClonerEffector/Source/ClonerEffector/Public/Cloner/CEClonerComponent.h b/Engine/Plugins/VirtualProduction/ClonerEffector/Source/ClonerEffector/Public/Cloner/CEClonerComponent.h
--- a/Engine/Plugins/VirtualProduction/ClonerEffector/Source/ClonerEffector/Public/Cloner/CEClonerComponent.h
+++ b/Engine/Plugins/VirtualProduction/ClonerEffector/Source/ClonerEffector/Public/Cloner/CEClonerComponent.h
@@ -21,7 +21,8 @@ UCLASS(MinimalAPI
 	, BlueprintType
 	, DisplayName = "Motion Design Cloner Component"
 	, AutoExpandCategories=(Cloner, Layout)
-	, HideCategories=(Transform, Niagara, Activation, Attachment, Randomness, Parameters))
+	, HideCategories=(Transform, Niagara, Activation, Attachment, Randomness, Parameters)
+	, meta = (BlueprintSpawnableComponent))
 class UCEClonerComponent : public UNiagaraComponent
 {
 	GENERATED_BODY()
diff --git a/Engine/Plugins/VirtualProduction/ClonerEffector/Source/ClonerEffector/Public/Effector/CEEffectorComponent.h b/Engine/Plugins/VirtualProduction/ClonerEffector/Source/ClonerEffector/Public/Effector/CEEffectorComponent.h
--- a/Engine/Plugins/VirtualProduction/ClonerEffector/Source/ClonerEffector/Public/Effector/CEEffectorComponent.h
+++ b/Engine/Plugins/VirtualProduction/ClonerEffector/Source/ClonerEffector/Public/Effector/CEEffectorComponent.h
@@ -21,7 +21,8 @@ UCLASS(
 	MinimalAPI
 	, BlueprintType
 	, HideCategories=(Tags,Activation,LOD,AssetUserData,Navigation,Rendering,Cooking,Physics,WorldPartition)
-	, DisplayName = "Motion Design Effector Component")
+	, DisplayName = "Motion Design Effector Component"
+	, meta=(BlueprintSpawnableComponent))
 class UCEEffectorComponent : public USceneComponent
 {
 	GENERATED_BODY()
