MotionDesign : ClonerEffector
- Fixed access violation in ClonerEffector plugin cause a crash
- Improved mesh conversion system

diff --git a/Engine/Plugins/VirtualProduction/ClonerEffector/Source/ClonerEffector/Private/Cloner/CEClonerComponent.cpp b/Engine/Plugins/VirtualProduction/ClonerEffector/Source/ClonerEffector/Private/Cloner/CEClonerComponent.cpp
--- a/Engine/Plugins/VirtualProduction/ClonerEffector/Source/ClonerEffector/Private/Cloner/CEClonerComponent.cpp
+++ b/Engine/Plugins/VirtualProduction/ClonerEffector/Source/ClonerEffector/Private/Cloner/CEClonerComponent.cpp
@@ -8,7 +8,6 @@
 #include "Cloner/Extensions/CEClonerExtensionBase.h"
 #include "Cloner/Layouts/CEClonerLayoutBase.h"
 #include "Components/BillboardComponent.h"
-#include "Components/DynamicMeshComponent.h"
 #include "Components/SkeletalMeshComponent.h"
 #include "Engine/Level.h"
 #include "Engine/StaticMeshActor.h"
@@ -196,8 +195,11 @@ void UCEClonerComponent::OnComponentDestroyed(bool bInDestroyingHierarchy)
 
 	FWorldDelegates::LevelAddedToWorld.RemoveAll(this);
 
-	FTSTicker::RemoveTicker(ClonerTickerHandle);
+	FTSTicker::GetCoreTicker().RemoveTicker(ClonerTickerHandle);
 	ClonerTickerHandle.Reset();
+
+	FTSTicker::GetCoreTicker().RemoveTicker(ConversionTickerHandle);
+	ConversionTickerHandle.Reset();
 }
 
 void UCEClonerComponent::UpdateClonerRenderState()
@@ -249,81 +251,77 @@ void UCEClonerComponent::UpdateDirtyMeshesAsync()
 
 	bClonerMeshesUpdating = true;
 
-	TSet<TWeakObjectPtr<AActor>> DirtyAttachments = ClonerTree.DirtyItemAttachments;
-	ClonerTree.DirtyItemAttachments.Empty(DirtyAttachments.Num());
-	
-	UE_LOG(LogCECloner, Verbose, TEXT("%s : Updating %i dirty actor meshes"), *GetOwner()->GetActorNameOrLabel(), DirtyAttachments.Num());
-
-	// Update baked dynamic meshes on other thread
-	TWeakObjectPtr<UCEClonerComponent> ThisWeak(this);
-	Async(EAsyncExecution::ThreadPool, [ThisWeak, DirtyAttachments]()
+	if (ConversionTickerHandle.IsValid())
 	{
-		UCEClonerComponent* This = ThisWeak.Get();
+		FTSTicker::GetCoreTicker().RemoveTicker(ConversionTickerHandle);
+	}
 
-		if (!This)
-		{
-			return;
-		}
+	UE_LOG(LogCECloner, Verbose, TEXT("%s : Updating %i dirty actor meshes..."), *GetOwner()->GetActorNameOrLabel(), ClonerTree.DirtyItemAttachments.Num());
 
-		// update actor baked dynamic meshes
-		bool bSuccess = true;
-		for (const TWeakObjectPtr<AActor>& Attachment : DirtyAttachments)
+	const double StartTime = FPlatformTime::Seconds();
+	const int32 DirtyItemCount = ClonerTree.DirtyItemAttachments.Num();
+	ConversionTickerHandle = FTSTicker::GetCoreTicker().AddTicker(FTickerDelegate::CreateWeakLambda(this, [This = this, StartTime, DirtyItemCount](float InDelta)->bool
 		{
-			AActor* DirtyActor = Attachment.Get();
-
-			if (!DirtyActor)
+			if (!This)
 			{
-				continue;
+				return false;
 			}
 
-			if (IsGarbageCollectingAndLockingUObjectHashTables())
-			{
-				bSuccess = false;
-				This->ClonerTree.DirtyItemAttachments.Add(Attachment);
-				continue;
-			}
+			FGCScopeGuard GCGuard;
+			const double StartTaskTime = FPlatformTime::Seconds();
 
-			This->UpdateActorDynamicMesh(DirtyActor);
-		}
+			{
+				// update actor baked dynamic meshes
+				for (TSet<TWeakObjectPtr<AActor>>::TIterator It(This->ClonerTree.DirtyItemAttachments); It; ++It)
+				{
+					const TWeakObjectPtr<AActor>& Attachment = *It;
 
-		// Create baked static mesh on main thread (required)
-		Async(EAsyncExecution::TaskGraphMainThread, [ThisWeak, &bSuccess]()
-		{
-			UCEClonerComponent* This = ThisWeak.Get();
+					if (AActor* DirtyActor = Attachment.Get())
+					{
+						This->UpdateActorDynamicMesh(DirtyActor);
+					}
 
-			if (!This)
-			{
-				return;
-			}
+					It.RemoveCurrent();
 
-			if (!bSuccess)
-			{
-				This->OnDirtyMeshesUpdated(false);
-				return;
+					if ((FPlatformTime::Seconds() - StartTaskTime) > InDelta)
+					{
+						// Continue next tick
+						return true;
+					}
+				}
 			}
 
-			// Update actors baked static mesh
-			for (const FCEClonerAttachmentRootItem& RootItem : This->ClonerTree.RootItemAttachments)
 			{
-				const UStaticMesh* RootStaticMesh = RootItem.MergedBakedMesh.Get();
-
-				if (!RootStaticMesh)
+				// Update actors baked static mesh
+				for (const FCEClonerAttachmentRootItem& RootItem : This->ClonerTree.RootItemAttachments)
 				{
-					if (IsGarbageCollectingAndLockingUObjectHashTables())
+					const UStaticMesh* RootStaticMesh = RootItem.MergedBakedMesh.Get();
+
+					if (!RootStaticMesh)
 					{
-						bSuccess = false;
-						break;
+						if (AActor* RootActor = RootItem.RootActor.Get())
+						{
+							This->UpdateActorStaticMesh(RootActor);
+						}
 					}
 
-					AActor* RootActor = RootItem.RootActor.Get();
-					This->UpdateActorStaticMesh(RootActor);
+					if ((FPlatformTime::Seconds() - StartTaskTime) > InDelta)
+					{
+						// Continue next tick
+						return true;
+					}
 				}
 			}
 
+			UE_LOG(LogCECloner, Verbose, TEXT("%s : Updated %i dirty actor meshes in %.3fs"), *This->GetOwner()->GetActorNameOrLabel(), DirtyItemCount, FPlatformTime::Seconds() - StartTime);
+
 			// update niagara asset
-			This->OnDirtyMeshesUpdated(bSuccess);
-		});
-	});
+			This->OnDirtyMeshesUpdated(/** Success */true);
+
+			// Stop
+			return false;
+		})
+	);
 }
 
 void UCEClonerComponent::UpdateActorDynamicMesh(AActor* InActor)
@@ -377,7 +375,7 @@ void UCEClonerComponent::UpdateActorDynamicMesh(AActor* InActor)
 		ClonerTree.MarkCacheOutdated(InActor);
 	}
 
-	UE_LOG(LogCECloner, Log, TEXT("%s : Updated actor dynamic mesh : %s - %i"), *ClonerActor->GetActorNameOrLabel(), *InActor->GetActorNameOrLabel(), Mesh->GetTriangleCount());
+	UE_LOG(LogCECloner, Verbose, TEXT("%s : Updated actor dynamic mesh : %s - %i"), *ClonerActor->GetActorNameOrLabel(), *InActor->GetActorNameOrLabel(), Mesh->GetTriangleCount());
 }
 
 void UCEClonerComponent::UpdateActorStaticMesh(AActor* InRootActor)
@@ -446,7 +444,7 @@ void UCEClonerComponent::UpdateActorStaticMesh(AActor* InRootActor)
 	RootItem->MeshSize = ClonerTree.GetAttachmentBounds(InRootActor, /** IncludeChildren */true).GetSize();
 #endif
 
-	UE_LOG(LogCECloner, Log, TEXT("%s : Updated actor static mesh : %s - %i - %i"), *ClonerActor->GetActorNameOrLabel(), *InRootActor->GetActorNameOrLabel(), Mesh->GetNumTriangles(/** LOD */0), AttachmentItems.Num());
+	UE_LOG(LogCECloner, Verbose, TEXT("%s : Updated actor static mesh : %s - %i - %i"), *ClonerActor->GetActorNameOrLabel(), *InRootActor->GetActorNameOrLabel(), Mesh->GetNumTriangles(/** LOD */0), AttachmentItems.Num());
 }
 
 void UCEClonerComponent::OnActorStaticMeshPostBuild(UStaticMesh* InMesh)
@@ -462,6 +460,7 @@ void UCEClonerComponent::OnActorStaticMeshPostBuild(UStaticMesh* InMesh)
 
 void UCEClonerComponent::OnDirtyMeshesUpdated(bool bInSuccess)
 {
+	ConversionTickerHandle.Reset();
 	bClonerMeshesUpdating = false;
 
 	// Update niagara parameters
diff --git a/Engine/Plugins/VirtualProduction/ClonerEffector/Source/ClonerEffector/Public/Cloner/CEClonerComponent.h b/Engine/Plugins/VirtualProduction/ClonerEffector/Source/ClonerEffector/Public/Cloner/CEClonerComponent.h
--- a/Engine/Plugins/VirtualProduction/ClonerEffector/Source/ClonerEffector/Public/Cloner/CEClonerComponent.h
+++ b/Engine/Plugins/VirtualProduction/ClonerEffector/Source/ClonerEffector/Public/Cloner/CEClonerComponent.h
@@ -231,7 +231,7 @@ public:
 		return Cast<InExtensionClass>(GetExtension(InExtensionClass::StaticClass()));
 	}
 
-	UFUNCTION(BlueprintCallable, Category="Cloner")
+	UFUNCTION(BlueprintCallable, Category="Cloner", meta=(DeterminesOutputType="InExtensionClass"))
 	UCEClonerExtensionBase* GetExtension(TSubclassOf<UCEClonerExtensionBase> InExtensionClass) const;
 
 	UCEClonerExtensionBase* GetExtension(FName InExtensionName) const;
@@ -453,6 +453,7 @@ private:
 	bool bClonerResourcesReady = false;
 
 	FTSTicker::FDelegateHandle ClonerTickerHandle;
+	FTSTicker::FDelegateHandle ConversionTickerHandle;
 
 #if WITH_EDITOR
 	/** Used for PECP */
diff --git a/Engine/Plugins/VirtualProduction/ClonerEffector/Source/ClonerEffector/Public/Effector/CEEffectorComponent.h b/Engine/Plugins/VirtualProduction/ClonerEffector/Source/ClonerEffector/Public/Effector/CEEffectorComponent.h
--- a/Engine/Plugins/VirtualProduction/ClonerEffector/Source/ClonerEffector/Public/Effector/CEEffectorComponent.h
+++ b/Engine/Plugins/VirtualProduction/ClonerEffector/Source/ClonerEffector/Public/Effector/CEEffectorComponent.h
@@ -173,7 +173,7 @@ public:
 	UFUNCTION(BlueprintPure, Category="Effector")
 	void GetActiveEffects(TArray<UCEEffectorEffectBase*>& OutEffects) const;
 
-	UFUNCTION(BlueprintCallable, Category="Effector")
+	UFUNCTION(BlueprintCallable, Category="Effector", meta=(DeterminesOutputType="InExtensionClass"))
 	UCEEffectorExtensionBase* GetExtension(TSubclassOf<UCEEffectorExtensionBase> InExtensionClass) const;
 
 	template<typename InExtensionClass UE_REQUIRES(std::is_base_of_v<UCEEffectorExtensionBase, InExtensionClass>)>
