MotionDesign : ClonerEffector
- Fixed details view property tab to commit unfocusing value widget of cloner and effector due to RequetRebuildChildren clearing focus

diff --git a/Engine/Plugins/VirtualProduction/ClonerEffector/Source/ClonerEffectorEditor/Private/CEEditorInputPreprocessor.cpp b/Engine/Plugins/VirtualProduction/ClonerEffector/Source/ClonerEffectorEditor/Private/CEEditorInputPreprocessor.cpp
new file mode 100644
--- /dev/null
+++ b/Engine/Plugins/VirtualProduction/ClonerEffector/Source/ClonerEffectorEditor/Private/CEEditorInputPreprocessor.cpp
@@ -0,0 +1,70 @@
+// Copyright Epic Games, Inc. All Rights Reserved.
+
+#include "CEEditorInputPreprocessor.h"
+
+#include "Containers/Ticker.h"
+#include "Framework/Application/SlateApplication.h"
+#include "Widgets/SWidget.h"
+#include "Widgets/SWindow.h"
+
+void FCEEditorInputPreprocessor::RefocusLastWidget()
+{
+	if (!FSlateApplication::IsInitialized() || !bShouldRefocus)
+	{
+		return;
+	}
+
+	bShouldRefocus = false;
+
+	const FSlateApplication& SlateApp = FSlateApplication::Get();
+
+	if (const TSharedPtr<SWidget> FocusedWidget = SlateApp.GetKeyboardFocusedWidget())
+	{
+		if (const TSharedPtr<SWindow> SlateWindow = SlateApp.FindWidgetWindow(FocusedWidget.ToSharedRef()))
+		{
+			if (TSharedPtr<FGenericWindow> NativeWindow = SlateWindow->GetNativeWindow())
+			{
+				const FVector2D ScreenPos = FocusedWidget->GetCachedGeometry().GetAbsolutePosition() + FocusedWidget->GetCachedGeometry().GetLocalSize() / 2;
+
+				FTSTicker::GetCoreTicker().AddTicker(FTickerDelegate::CreateSPLambda(this, [ScreenPos](float)->bool
+				{
+					if (!FSlateApplication::IsInitialized())
+					{
+						return false;
+					}
+
+					FSlateApplication& SlateApp = FSlateApplication::Get();
+
+					// Hit-test at the event position
+					FWidgetPath WidgetPath = SlateApp.LocateWindowUnderMouse(ScreenPos, SlateApp.GetInteractiveTopLevelWindows());
+
+					if (WidgetPath.IsValid() && WidgetPath.Widgets.Num() > 0)
+					{
+						// Last in the path is the most specific widget under the cursor
+						TSharedRef<SWidget> TargetWidget = WidgetPath.Widgets.Last().Widget;
+						SlateApp.SetAllUserFocus(TargetWidget);
+					}
+
+					return false;
+				}));
+			}
+		}
+	}
+}
+
+void FCEEditorInputPreprocessor::Register()
+{
+	if (!bRegistered && FSlateApplication::IsInitialized())
+	{
+		bRegistered = FSlateApplication::Get().RegisterInputPreProcessor(SharedThis(this));
+	}
+}
+
+void FCEEditorInputPreprocessor::Unregister()
+{
+	if (bRegistered && FSlateApplication::IsInitialized())
+	{
+		FSlateApplication::Get().UnregisterInputPreProcessor(SharedThis(this));
+		bRegistered = false;
+	}
+}
diff --git a/Engine/Plugins/VirtualProduction/ClonerEffector/Source/ClonerEffectorEditor/Private/CEEditorInputPreprocessor.h b/Engine/Plugins/VirtualProduction/ClonerEffector/Source/ClonerEffectorEditor/Private/CEEditorInputPreprocessor.h
new file mode 100644
--- /dev/null
+++ b/Engine/Plugins/VirtualProduction/ClonerEffector/Source/ClonerEffectorEditor/Private/CEEditorInputPreprocessor.h
@@ -0,0 +1,50 @@
+// Copyright Epic Games, Inc. All Rights Reserved.
+
+#pragma once
+
+#include "Framework/Application/IInputProcessor.h"
+#include "Input/Events.h"
+#include "InputCoreTypes.h"
+#include "Templates/SharedPointer.h"
+
+/**
+ * This is a fix used to refocus last selected widget in the details panel
+ * Since a details refresh causes property value widget to lose focus
+ * This should be removed in the future once PropertyValueWidget can be refocused through the API or refresh does not cause a focus lost
+ */
+class FCEEditorInputPreprocessor : public IInputProcessor, public TSharedFromThis<FCEEditorInputPreprocessor>
+{
+public:
+	//~ Begin IInputProcessor
+	virtual void Tick(const float InDeltaTime, FSlateApplication& InSlateApp, TSharedRef<ICursor> InCursor) override
+	{
+	};
+
+	virtual bool HandleKeyDownEvent(FSlateApplication& InSlateApp, const FKeyEvent& InKeyEvent) override
+	{
+		bShouldRefocus = InKeyEvent.GetKey() == EKeys::Tab || InKeyEvent.GetKey() == EKeys::Enter;
+		// Return false to allow other processors to handle the event
+		return false;
+	}
+
+	virtual bool HandleMouseButtonDownEvent(FSlateApplication& InSlateApp, const FPointerEvent& InMouseEvent) override
+	{
+		bShouldRefocus = false;
+		// Return false to allow other processors to handle the event
+		return false;
+	}
+
+	virtual const TCHAR* GetDebugName() const override
+	{
+		return TEXT("ClonerEffector");
+	}
+	//~ End IInputProcessor
+
+	void RefocusLastWidget();
+	void Register();
+	void Unregister();
+
+private:
+	bool bShouldRefocus = false;
+	bool bRegistered = false;
+};
\ No newline at end of file
diff --git a/Engine/Plugins/VirtualProduction/ClonerEffector/Source/ClonerEffectorEditor/Private/CEEditorModule.cpp b/Engine/Plugins/VirtualProduction/ClonerEffector/Source/ClonerEffectorEditor/Private/CEEditorModule.cpp
--- a/Engine/Plugins/VirtualProduction/ClonerEffector/Source/ClonerEffectorEditor/Private/CEEditorModule.cpp
+++ b/Engine/Plugins/VirtualProduction/ClonerEffector/Source/ClonerEffectorEditor/Private/CEEditorModule.cpp
@@ -2,6 +2,7 @@
 
 #include "CEEditorModule.h"
 
+#include "CEEditorInputPreprocessor.h"
 #include "CEEditorThrottleManager.h"
 #include "Cloner/CEClonerComponent.h"
 #include "Cloner/CEClonerActor.h"
@@ -35,17 +36,19 @@ void FCEEditorModule::StartupModule()
 	FPropertyEditorModule& PropertyModule = FModuleManager::LoadModuleChecked<FPropertyEditorModule>(TEXT("PropertyEditor"));
 
 	// Cloner customization
-	PropertyModule.RegisterCustomClassLayout(ACEClonerActor::StaticClass()->GetFName(), FOnGetDetailCustomizationInstance::CreateStatic(&FCEEditorClonerActorDetailCustomization::MakeInstance));
-	PropertyModule.RegisterCustomClassLayout(UCEClonerComponent::StaticClass()->GetFName(), FOnGetDetailCustomizationInstance::CreateStatic(&FCEEditorClonerComponentDetailCustomization::MakeInstance));
-	PropertyModule.RegisterCustomClassLayout(UCEClonerEffectorExtension::StaticClass()->GetFName(), FOnGetDetailCustomizationInstance::CreateStatic(&FCEEditorClonerEffectorExtensionDetailCustomization::MakeInstance));
-	PropertyModule.RegisterCustomClassLayout(UCEClonerSplineLayout::StaticClass()->GetFName(), FOnGetDetailCustomizationInstance::CreateStatic(&FCEEditorClonerSplineLayoutDetailCustomization::MakeInstance));
-	PropertyModule.RegisterCustomClassLayout(UCEClonerMeshLayout::StaticClass()->GetFName(), FOnGetDetailCustomizationInstance::CreateStatic(&FCEEditorClonerMeshLayoutDetailCustomization::MakeInstance));
-	PropertyModule.RegisterCustomClassLayout(UCEClonerLifetimeExtension::StaticClass()->GetFName(), FOnGetDetailCustomizationInstance::CreateStatic(&FCEEditorClonerLifetimeExtensionDetailCustomization::MakeInstance));
+	PropertyModule.RegisterCustomClassLayout(CustomizationNames.Add_GetRef(ACEClonerActor::StaticClass()->GetFName()), FOnGetDetailCustomizationInstance::CreateStatic(&FCEEditorClonerActorDetailCustomization::MakeInstance));
+	PropertyModule.RegisterCustomClassLayout(CustomizationNames.Add_GetRef(UCEClonerComponent::StaticClass()->GetFName()), FOnGetDetailCustomizationInstance::CreateStatic(&FCEEditorClonerComponentDetailCustomization::MakeInstance));
+	PropertyModule.RegisterCustomClassLayout(CustomizationNames.Add_GetRef(UCEClonerEffectorExtension::StaticClass()->GetFName()), FOnGetDetailCustomizationInstance::CreateStatic(&FCEEditorClonerEffectorExtensionDetailCustomization::MakeInstance));
+	PropertyModule.RegisterCustomClassLayout(CustomizationNames.Add_GetRef(UCEClonerSplineLayout::StaticClass()->GetFName()), FOnGetDetailCustomizationInstance::CreateStatic(&FCEEditorClonerSplineLayoutDetailCustomization::MakeInstance));
+	PropertyModule.RegisterCustomClassLayout(CustomizationNames.Add_GetRef(UCEClonerMeshLayout::StaticClass()->GetFName()), FOnGetDetailCustomizationInstance::CreateStatic(&FCEEditorClonerMeshLayoutDetailCustomization::MakeInstance));
+	PropertyModule.RegisterCustomClassLayout(CustomizationNames.Add_GetRef(UCEClonerLifetimeExtension::StaticClass()->GetFName()), FOnGetDetailCustomizationInstance::CreateStatic(&FCEEditorClonerLifetimeExtensionDetailCustomization::MakeInstance));
+
+	InputPreprocessor = MakeShared<FCEEditorInputPreprocessor>();
 
 	// Effector customization
-	PropertyModule.RegisterCustomClassLayout(ACEEffectorActor::StaticClass()->GetFName(), FOnGetDetailCustomizationInstance::CreateStatic(&FCEEditorEffectorActorDetailCustomization::MakeInstance));
-	PropertyModule.RegisterCustomClassLayout(UCEEffectorComponent::StaticClass()->GetFName(), FOnGetDetailCustomizationInstance::CreateStatic(&FCEEditorEffectorComponentDetailCustomization::MakeInstance));
-	PropertyModule.RegisterCustomClassLayout(UCEEffectorBoundType::StaticClass()->GetFName(), FOnGetDetailCustomizationInstance::CreateStatic(&FCEEditorEffectorTypeDetailCustomization::MakeInstance));
+	PropertyModule.RegisterCustomClassLayout(CustomizationNames.Add_GetRef(ACEEffectorActor::StaticClass()->GetFName()), FOnGetDetailCustomizationInstance::CreateStatic(&FCEEditorEffectorActorDetailCustomization::MakeInstance));
+	PropertyModule.RegisterCustomClassLayout(CustomizationNames.Add_GetRef(UCEEffectorComponent::StaticClass()->GetFName()), FOnGetDetailCustomizationInstance::CreateStatic(&FCEEditorEffectorComponentDetailCustomization::MakeInstance));
+	PropertyModule.RegisterCustomClassLayout(CustomizationNames.Add_GetRef(UCEEffectorBoundType::StaticClass()->GetFName()), FOnGetDetailCustomizationInstance::CreateStatic(&FCEEditorEffectorTypeDetailCustomization::MakeInstance));
 
 	// Custom cloner track
 	ISequencerModule& SequencerModule = FModuleManager::LoadModuleChecked<ISequencerModule>(TEXT("Sequencer"));
@@ -58,23 +61,14 @@ void FCEEditorModule::StartupModule()
 
 void FCEEditorModule::ShutdownModule()
 {
-	if (UObjectInitialized())
+	if (FPropertyEditorModule* PropertyModule = FModuleManager::GetModulePtr<FPropertyEditorModule>("PropertyEditor"))
 	{
-		if (FPropertyEditorModule* PropertyModule = FModuleManager::GetModulePtr<FPropertyEditorModule>("PropertyEditor"))
+		for (const FName CustomizationName : CustomizationNames)
 		{
-			// Cloner customization
-			PropertyModule->UnregisterCustomClassLayout(ACEClonerActor::StaticClass()->GetFName());
-			PropertyModule->UnregisterCustomClassLayout(UCEClonerComponent::StaticClass()->GetFName());
-			PropertyModule->UnregisterCustomClassLayout(UCEClonerEffectorExtension::StaticClass()->GetFName());
-			PropertyModule->UnregisterCustomClassLayout(UCEClonerSplineLayout::StaticClass()->GetFName());
-			PropertyModule->UnregisterCustomClassLayout(UCEClonerMeshLayout::StaticClass()->GetFName());
-			PropertyModule->UnregisterCustomClassLayout(UCEClonerLifetimeExtension::StaticClass()->GetFName());
-
-			// Effector customization
-			PropertyModule->UnregisterCustomClassLayout(ACEEffectorActor::StaticClass()->GetFName());
-			PropertyModule->UnregisterCustomClassLayout(UCEEffectorComponent::StaticClass()->GetFName());
-			PropertyModule->UnregisterCustomClassLayout(UCEEffectorBoundType::StaticClass()->GetFName());
+			PropertyModule->UnregisterCustomClassLayout(CustomizationName);
 		}
+
+		CustomizationNames.Empty();
 	}
 
 	// Custom cloner track
@@ -84,8 +78,24 @@ void FCEEditorModule::ShutdownModule()
 		ClonerTrackCreateEditorHandle.Reset();
 	}
 
+	if (InputPreprocessor.IsValid())
+	{
+		InputPreprocessor->Unregister();
+		InputPreprocessor.Reset();
+	}
+
 	ThrottleManager.Reset();
 	ThrottleManager = nullptr;
 }
 
+TSharedPtr<FCEEditorInputPreprocessor> FCEEditorModule::GetInputPreprocessor()
+{
+	if (FCEEditorModule* EditorModule = FModuleManager::GetModulePtr<FCEEditorModule>(UE_MODULE_NAME))
+	{
+		return EditorModule->InputPreprocessor;
+	}
+
+	return nullptr;
+}
+
 IMPLEMENT_MODULE(FCEEditorModule, ClonerEffectorEditor)
diff --git a/Engine/Plugins/VirtualProduction/ClonerEffector/Source/ClonerEffectorEditor/Private/CEEditorModule.h b/Engine/Plugins/VirtualProduction/ClonerEffector/Source/ClonerEffectorEditor/Private/CEEditorModule.h
--- a/Engine/Plugins/VirtualProduction/ClonerEffector/Source/ClonerEffectorEditor/Private/CEEditorModule.h
+++ b/Engine/Plugins/VirtualProduction/ClonerEffector/Source/ClonerEffectorEditor/Private/CEEditorModule.h
@@ -4,6 +4,7 @@
 
 #include "Modules/ModuleManager.h"
 
+class FCEEditorInputPreprocessor;
 class FCEEditorThrottleManager;
 
 class FCEEditorModule : public IModuleInterface
@@ -14,7 +15,11 @@ public:
     virtual void ShutdownModule() override;
 	//~ End IModuleInterface
 
+	static TSharedPtr<FCEEditorInputPreprocessor> GetInputPreprocessor();
+
 protected:
 	FDelegateHandle ClonerTrackCreateEditorHandle;
 	TSharedPtr<FCEEditorThrottleManager> ThrottleManager;
+	TSharedPtr<FCEEditorInputPreprocessor> InputPreprocessor;
+	TArray<FName> CustomizationNames;
 };
diff --git a/Engine/Plugins/VirtualProduction/ClonerEffector/Source/ClonerEffectorEditor/Private/Cloner/Customizations/CEEditorClonerComponentDetailCustomization.cpp b/Engine/Plugins/VirtualProduction/ClonerEffector/Source/ClonerEffectorEditor/Private/Cloner/Customizations/CEEditorClonerComponentDetailCustomization.cpp
--- a/Engine/Plugins/VirtualProduction/ClonerEffector/Source/ClonerEffectorEditor/Private/Cloner/Customizations/CEEditorClonerComponentDetailCustomization.cpp
+++ b/Engine/Plugins/VirtualProduction/ClonerEffector/Source/ClonerEffectorEditor/Private/Cloner/Customizations/CEEditorClonerComponentDetailCustomization.cpp
@@ -2,6 +2,8 @@
 
 #include "Cloner/Customizations/CEEditorClonerComponentDetailCustomization.h"
 
+#include "CEEditorInputPreprocessor.h"
+#include "CEEditorModule.h"
 #include "Cloner/CEClonerComponent.h"
 #include "Cloner/Extensions/CEClonerExtensionBase.h"
 #include "Cloner/Layouts/CEClonerLayoutBase.h"
@@ -30,6 +32,7 @@ void FCEEditorClonerComponentDetailCustomization::CustomizeDetails(IDetailLayout
 
 	// Remove niagara utilities
 	InDetailBuilder.HideCategory(TEXT("NiagaraComponent_Utilities"));
+	InDetailBuilder.HideCategory(TEXT("Niagara Utilities"));
 
 	// Remove details sections by hiding specific categories
 	InDetailBuilder.HideCategory(TEXT("LOD"));	
@@ -105,6 +108,13 @@ void FCEEditorClonerComponentDetailCustomization::CustomizeDetails(IDetailLayout
 	AddParams.CreateCategoryNodes(false);
 	AddParams.HideRootObjectNode(true);
 
+	const TSharedPtr<FCEEditorInputPreprocessor> InputPreprocessor = FCEEditorModule::GetInputPreprocessor();
+
+	if (InputPreprocessor.IsValid())
+	{
+		InputPreprocessor->Unregister();
+	}
+
 	for (const TPair<FName,FDetailsCategoryData>& CategoryToObjectsPair : CategoryToData)
 	{
 		const FName CategoryName = CategoryToObjectsPair.Key;
@@ -125,7 +135,13 @@ void FCEEditorClonerComponentDetailCustomization::CustomizeDetails(IDetailLayout
 			TSharedPtr<IPropertyHandle> ObjectPropertyHandle = ObjectRow->GetPropertyHandle();
 
 			// Fix for EditConditionHides not appearing when condition is met due to AddExternalObjects not rebuilding children
-			ObjectPropertyHandle->SetOnChildPropertyValueChangedWithData(TDelegate<void(const FPropertyChangedEvent&)>::CreateStatic(&FCEEditorClonerComponentDetailCustomization::OnChildPropertyChanged, ObjectPropertyHandle->AsWeak()));
+			ObjectPropertyHandle->SetOnChildPropertyValueChangedWithData(TDelegate<void(const FPropertyChangedEvent&)>::CreateStatic(&FCEEditorClonerComponentDetailCustomization::OnChildPropertyChanged, ObjectPropertyHandle.ToWeakPtr()));
+
+			// Fix to refocus last widget when TAB is used since RequestRebuildChildren will lose focus
+			if (InputPreprocessor.IsValid())
+			{
+				InputPreprocessor->Register();
+			}
 		}
 	}
 
@@ -275,13 +291,28 @@ void FCEEditorClonerComponentDetailCustomization::Init()
 
 void FCEEditorClonerComponentDetailCustomization::OnChildPropertyChanged(const FPropertyChangedEvent& InEvent, TWeakPtr<IPropertyHandle> InParentHandleWeak)
 {
-	if (const TSharedPtr<IPropertyHandle> PropertyHandle = InParentHandleWeak.Pin())
+	const TSharedPtr<IPropertyHandle> PropertyHandle = InParentHandleWeak.Pin();
+
+	if (!PropertyHandle)
 	{
-		if (InEvent.ChangeType != EPropertyChangeType::Interactive)
-		{
-			PropertyHandle->RequestRebuildChildren();
-		}
+		return;
+	}
+
+	if (InEvent.ChangeType == EPropertyChangeType::Interactive)
+	{
+		return;
 	}
+
+	const TSharedPtr<FCEEditorInputPreprocessor> InputPreprocessor = FCEEditorModule::GetInputPreprocessor();
+
+	if (InputPreprocessor.IsValid())
+	{
+		// Refocus last focused widget next tick because rebuilding children loses focus
+		InputPreprocessor->RefocusLastWidget();
+		InputPreprocessor->Unregister();
+	}
+
+	PropertyHandle->RequestRebuildChildren();
 }
 
 void FCEEditorClonerComponentDetailCustomization::OnClonerLayoutLoaded(UCEClonerComponent* InCloner, UCEClonerLayoutBase* InLayout)
diff --git a/Engine/Plugins/VirtualProduction/ClonerEffector/Source/ClonerEffectorEditor/Private/Effector/Customizations/CEEditorEffectorComponentDetailCustomization.cpp b/Engine/Plugins/VirtualProduction/ClonerEffector/Source/ClonerEffectorEditor/Private/Effector/Customizations/CEEditorEffectorComponentDetailCustomization.cpp
--- a/Engine/Plugins/VirtualProduction/ClonerEffector/Source/ClonerEffectorEditor/Private/Effector/Customizations/CEEditorEffectorComponentDetailCustomization.cpp
+++ b/Engine/Plugins/VirtualProduction/ClonerEffector/Source/ClonerEffectorEditor/Private/Effector/Customizations/CEEditorEffectorComponentDetailCustomization.cpp
@@ -2,6 +2,8 @@
 
 #include "Effector/Customizations/CEEditorEffectorComponentDetailCustomization.h"
 
+#include "CEEditorInputPreprocessor.h"
+#include "CEEditorModule.h"
 #include "DetailLayoutBuilder.h"
 #include "DetailWidgetRow.h"
 #include "Effector/CEEffectorComponent.h"
@@ -98,6 +100,13 @@ void FCEEditorEffectorComponentDetailCustomization::CustomizeDetails(IDetailLayo
 	AddParams.CreateCategoryNodes(false);
 	AddParams.HideRootObjectNode(true);
 
+	const TSharedPtr<FCEEditorInputPreprocessor> InputPreprocessor = FCEEditorModule::GetInputPreprocessor();
+
+	if (InputPreprocessor.IsValid())
+	{
+		InputPreprocessor->Unregister();
+	}
+
 	for (const TPair<FName,FDetailsCategoryData>& CategoryToObjectsPair : CategoryToData)
 	{
 		const FName CategoryName = CategoryToObjectsPair.Key;
@@ -119,6 +128,12 @@ void FCEEditorEffectorComponentDetailCustomization::CustomizeDetails(IDetailLayo
 
 			// Fix for EditConditionHides not appearing when condition is met due to AddExternalObjects not rebuilding children
 			ObjectPropertyHandle->SetOnChildPropertyValueChangedWithData(TDelegate<void(const FPropertyChangedEvent&)>::CreateStatic(&FCEEditorEffectorComponentDetailCustomization::OnChildPropertyChanged, ObjectPropertyHandle->AsWeak()));
+
+			// Fix to refocus last widget when TAB is used since RequestRebuildChildren will lose focus
+			if (InputPreprocessor.IsValid())
+			{
+				InputPreprocessor->Register();
+			}
 		}
 	}
 }
@@ -145,13 +160,28 @@ void FCEEditorEffectorComponentDetailCustomization::OnPropertyChanged(const FPro
 
 void FCEEditorEffectorComponentDetailCustomization::OnChildPropertyChanged(const FPropertyChangedEvent& InEvent, TWeakPtr<IPropertyHandle> InParentHandleWeak)
 {
-	if (const TSharedPtr<IPropertyHandle> PropertyHandle = InParentHandleWeak.Pin())
+	const TSharedPtr<IPropertyHandle> PropertyHandle = InParentHandleWeak.Pin();
+
+	if (!PropertyHandle)
 	{
-		if (InEvent.ChangeType != EPropertyChangeType::Interactive)
-		{
-			PropertyHandle->RequestRebuildChildren();
-		}
+		return;
 	}
+
+	if (InEvent.ChangeType == EPropertyChangeType::Interactive)
+	{
+		return;
+	}
+
+	const TSharedPtr<FCEEditorInputPreprocessor> InputPreprocessor = FCEEditorModule::GetInputPreprocessor();
+
+	if (InputPreprocessor.IsValid())
+	{
+		// Refocus last focused widget next tick because rebuilding children loses focus
+		InputPreprocessor->RefocusLastWidget();
+		InputPreprocessor->Unregister();
+	}
+
+	PropertyHandle->RequestRebuildChildren();
 }
 
 #undef LOCTEXT_NAMESPACE
