MotionDesign : ClonerEffector
- Added metadata to only refresh children in details view when changed property has the metadata "RefreshPropertyView"

diff --git a/Engine/Plugins/VirtualProduction/ClonerEffector/Source/ClonerEffector/Public/Cloner/Extensions/CEClonerCollisionExtension.h b/Engine/Plugins/VirtualProduction/ClonerEffector/Source/ClonerEffector/Public/Cloner/Extensions/CEClonerCollisionExtension.h
--- a/Engine/Plugins/VirtualProduction/ClonerEffector/Source/ClonerEffector/Public/Cloner/Extensions/CEClonerCollisionExtension.h
+++ b/Engine/Plugins/VirtualProduction/ClonerEffector/Source/ClonerEffector/Public/Cloner/Extensions/CEClonerCollisionExtension.h
@@ -128,11 +128,11 @@ protected:
 	 * Allow particles to react to surface by using the distance field,
 	 * ensure the mesh you want particle to collide with is tick enough or increase its distance field resolution scale in static mesh editor 
 	 */
-	UPROPERTY(EditInstanceOnly, Setter="SetSurfaceCollisionEnabled", Getter="GetSurfaceCollisionEnabled", Category="Collisions")
+	UPROPERTY(EditInstanceOnly, Setter="SetSurfaceCollisionEnabled", Getter="GetSurfaceCollisionEnabled", Category="Collisions", meta=(RefreshPropertyView))
 	bool bSurfaceCollisionEnabled = false;
 
 	/** Allow particles to react to other emitter particles, uses a neighbor grid to detect collision */
-	UPROPERTY(EditInstanceOnly, Setter="SetParticleCollisionEnabled", Getter="GetParticleCollisionEnabled", Category="Collisions")
+	UPROPERTY(EditInstanceOnly, Setter="SetParticleCollisionEnabled", Getter="GetParticleCollisionEnabled", Category="Collisions", meta=(RefreshPropertyView))
 	bool bParticleCollisionEnabled = false;
 
 	/** Recalculate accurate velocity after collision is updated */
@@ -152,7 +152,7 @@ protected:
 	FVector CollisionGridSize = FVector(5000.f);
 
 	/** Collision radius calculation mode */
-	UPROPERTY(EditInstanceOnly, Setter, Getter, Category="Collisions", meta=(EditCondition="bSurfaceCollisionEnabled || bParticleCollisionEnabled", EditConditionHides))
+	UPROPERTY(EditInstanceOnly, Setter, Getter, Category="Collisions", meta=(EditCondition="bSurfaceCollisionEnabled || bParticleCollisionEnabled", EditConditionHides, RefreshPropertyView))
 	ECEClonerCollisionRadiusMode CollisionRadiusMode = ECEClonerCollisionRadiusMode::ExtentLength;
 
 	/** Radius expected around each particle for collision, order matches attachment index */
diff --git a/Engine/Plugins/VirtualProduction/ClonerEffector/Source/ClonerEffector/Public/Cloner/Extensions/CEClonerConstraintExtension.h b/Engine/Plugins/VirtualProduction/ClonerEffector/Source/ClonerEffector/Public/Cloner/Extensions/CEClonerConstraintExtension.h
--- a/Engine/Plugins/VirtualProduction/ClonerEffector/Source/ClonerEffector/Public/Cloner/Extensions/CEClonerConstraintExtension.h
+++ b/Engine/Plugins/VirtualProduction/ClonerEffector/Source/ClonerEffector/Public/Cloner/Extensions/CEClonerConstraintExtension.h
@@ -171,7 +171,7 @@ protected:
 	virtual bool IsLayoutSupported(const UCEClonerLayoutBase* InLayout) const override;
 	//~ End UCEClonerExtensionBase
 
-	UPROPERTY(EditInstanceOnly, Setter, Getter, Category="Constraint")
+	UPROPERTY(EditInstanceOnly, Setter, Getter, Category="Constraint", meta=(RefreshPropertyView))
 	ECEClonerGridConstraint Constraint = ECEClonerGridConstraint::None;
 
 	UPROPERTY(EditInstanceOnly, Setter="SetInvertConstraint", Getter="GetInvertConstraint", Category="Constraint", meta=(EditCondition="Constraint != ECEClonerGridConstraint::None", EditConditionHides))
diff --git a/Engine/Plugins/VirtualProduction/ClonerEffector/Source/ClonerEffector/Public/Cloner/Extensions/CEClonerDisplacementExtension.h b/Engine/Plugins/VirtualProduction/ClonerEffector/Source/ClonerEffector/Public/Cloner/Extensions/CEClonerDisplacementExtension.h
--- a/Engine/Plugins/VirtualProduction/ClonerEffector/Source/ClonerEffector/Public/Cloner/Extensions/CEClonerDisplacementExtension.h
+++ b/Engine/Plugins/VirtualProduction/ClonerEffector/Source/ClonerEffector/Public/Cloner/Extensions/CEClonerDisplacementExtension.h
@@ -135,7 +135,7 @@ protected:
 	virtual bool IsLayoutSupported(const UCEClonerLayoutBase* InLayout) const override;
 	//~ End UCEClonerExtensionBase
 
-	UPROPERTY(EditInstanceOnly, Setter="SetDisplacementEnabled", Getter="GetDisplacementEnabled", DisplayName="Enabled", Category="Displacement")
+	UPROPERTY(EditInstanceOnly, Setter="SetDisplacementEnabled", Getter="GetDisplacementEnabled", DisplayName="Enabled", Category="Displacement", meta=(RefreshPropertyView))
 	bool bDisplacementEnabled = false;
 
 	UPROPERTY(EditInstanceOnly, Setter="SetDisplacementInvert", Getter="GetDisplacementInvert", DisplayName="Invert", Category="Displacement", meta=(EditCondition="bDisplacementEnabled", EditConditionHides))
diff --git a/Engine/Plugins/VirtualProduction/ClonerEffector/Source/ClonerEffector/Public/Cloner/Extensions/CEClonerEmitterSpawnExtension.h b/Engine/Plugins/VirtualProduction/ClonerEffector/Source/ClonerEffector/Public/Cloner/Extensions/CEClonerEmitterSpawnExtension.h
--- a/Engine/Plugins/VirtualProduction/ClonerEffector/Source/ClonerEffector/Public/Cloner/Extensions/CEClonerEmitterSpawnExtension.h
+++ b/Engine/Plugins/VirtualProduction/ClonerEffector/Source/ClonerEffector/Public/Cloner/Extensions/CEClonerEmitterSpawnExtension.h
@@ -104,7 +104,7 @@ protected:
 	//~ End UCEClonerExtensionBase
 
 	/** How many times do we spawn clones */
-	UPROPERTY(EditInstanceOnly, Setter, Getter, DisplayName="Emission Mode", Category="Spawn")
+	UPROPERTY(EditInstanceOnly, Setter, Getter, DisplayName="Emission Mode", Category="Spawn", meta=(RefreshPropertyView))
 	ECEClonerSpawnLoopMode SpawnLoopMode = ECEClonerSpawnLoopMode::Once;
 
 	/** Amount of spawn iterations for clones */
@@ -112,7 +112,7 @@ protected:
 	int32 SpawnLoopIterations = 1;
 
 	/** How does spawn occurs */
-	UPROPERTY(EditInstanceOnly, Setter, Getter, DisplayName="Emission Style", Category="Spawn", meta=(EditCondition="SpawnLoopMode != ECEClonerSpawnLoopMode::Once", EditConditionHides))
+	UPROPERTY(EditInstanceOnly, Setter, Getter, DisplayName="Emission Style", Category="Spawn", meta=(EditCondition="SpawnLoopMode != ECEClonerSpawnLoopMode::Once", EditConditionHides, RefreshPropertyView))
 	ECEClonerSpawnBehaviorMode SpawnBehaviorMode = ECEClonerSpawnBehaviorMode::Instant;
 
 	/** Interval/Duration of spawn for clones */
diff --git a/Engine/Plugins/VirtualProduction/ClonerEffector/Source/ClonerEffector/Public/Cloner/Extensions/CEClonerLifetimeExtension.h b/Engine/Plugins/VirtualProduction/ClonerEffector/Source/ClonerEffector/Public/Cloner/Extensions/CEClonerLifetimeExtension.h
--- a/Engine/Plugins/VirtualProduction/ClonerEffector/Source/ClonerEffector/Public/Cloner/Extensions/CEClonerLifetimeExtension.h
+++ b/Engine/Plugins/VirtualProduction/ClonerEffector/Source/ClonerEffector/Public/Cloner/Extensions/CEClonerLifetimeExtension.h
@@ -84,7 +84,7 @@ protected:
 	void OnLifetimeScaleCurveChanged();
 
 	/** Do we destroy the clones after a specific duration */
-	UPROPERTY(EditInstanceOnly, Setter="SetLifetimeEnabled", Getter="GetLifetimeEnabled", DisplayName="Enabled", Category="Lifetime")
+	UPROPERTY(EditInstanceOnly, Setter="SetLifetimeEnabled", Getter="GetLifetimeEnabled", DisplayName="Enabled", Category="Lifetime", meta=(RefreshPropertyView))
 	bool bLifetimeEnabled = false;
 
 	/** Minimum lifetime for a clone */
diff --git a/Engine/Plugins/VirtualProduction/ClonerEffector/Source/ClonerEffector/Public/Cloner/Extensions/CEClonerMeshRendererExtension.h b/Engine/Plugins/VirtualProduction/ClonerEffector/Source/ClonerEffector/Public/Cloner/Extensions/CEClonerMeshRendererExtension.h
--- a/Engine/Plugins/VirtualProduction/ClonerEffector/Source/ClonerEffector/Public/Cloner/Extensions/CEClonerMeshRendererExtension.h
+++ b/Engine/Plugins/VirtualProduction/ClonerEffector/Source/ClonerEffector/Public/Cloner/Extensions/CEClonerMeshRendererExtension.h
@@ -135,7 +135,7 @@ protected:
 	bool bVisualizeEffectors = false;
 
 	/** Whether to override meshes materials with another material */
-	UPROPERTY(EditInstanceOnly, Setter="SetUseOverrideMaterial", Getter="GetUseOverrideMaterial", Category="Renderer")
+	UPROPERTY(EditInstanceOnly, Setter="SetUseOverrideMaterial", Getter="GetUseOverrideMaterial", Category="Renderer", meta=(RefreshPropertyView))
 	bool bUseOverrideMaterial = false;
 
 	/** The override materials that will be set instead of meshes materials, bVisualizeEffectors must be disabled */
diff --git a/Engine/Plugins/VirtualProduction/ClonerEffector/Source/ClonerEffector/Public/Cloner/Extensions/CEClonerRangeExtension.h b/Engine/Plugins/VirtualProduction/ClonerEffector/Source/ClonerEffector/Public/Cloner/Extensions/CEClonerRangeExtension.h
--- a/Engine/Plugins/VirtualProduction/ClonerEffector/Source/ClonerEffector/Public/Cloner/Extensions/CEClonerRangeExtension.h
+++ b/Engine/Plugins/VirtualProduction/ClonerEffector/Source/ClonerEffector/Public/Cloner/Extensions/CEClonerRangeExtension.h
@@ -126,11 +126,11 @@ protected:
 	//~ End UCEClonerExtensionBase
 
 	/** Use random range transforms for each clones */
-	UPROPERTY(EditInstanceOnly, Setter="SetRangeEnabled", Getter="GetRangeEnabled", DisplayName="Enabled", Category="Range")
+	UPROPERTY(EditInstanceOnly, Setter="SetRangeEnabled", Getter="GetRangeEnabled", DisplayName="Enabled", Category="Range", meta=(RefreshPropertyView))
 	bool bRangeEnabled = false;
 
 	/** Mirrors max offset and rotation values in min offset and rotation */
-	UPROPERTY(EditInstanceOnly, Setter="SetRangeMirrored", Getter="GetRangeMirrored", DisplayName="Mirror Min & Max", Category="Range", meta=(EditCondition="bRangeEnabled", EditConditionHides))
+	UPROPERTY(EditInstanceOnly, Setter="SetRangeMirrored", Getter="GetRangeMirrored", DisplayName="Mirror Min & Max", Category="Range", meta=(EditCondition="bRangeEnabled", EditConditionHides, RefreshPropertyView))
 	bool bRangeMirrored = true;
 
 	UPROPERTY(EditInstanceOnly, Setter, Getter, DisplayName="OffsetMin", Category="Range", meta=(EditCondition="bRangeEnabled && !bRangeMirrored", EditConditionHides))
@@ -145,7 +145,7 @@ protected:
 	UPROPERTY(EditInstanceOnly, Setter, Getter, DisplayName="RotationMax", Category="Range", meta=(EditCondition="bRangeEnabled", EditConditionHides))
 	FRotator RangeRotationMax = FRotator::ZeroRotator;
 
-	UPROPERTY(EditInstanceOnly, Setter="SetRangeScaleUniform", Getter="GetRangeScaleUniform", DisplayName="ScaleUniformEnabled", Category="Range", meta=(EditCondition="bRangeEnabled", EditConditionHides))
+	UPROPERTY(EditInstanceOnly, Setter="SetRangeScaleUniform", Getter="GetRangeScaleUniform", DisplayName="ScaleUniformEnabled", Category="Range", meta=(EditCondition="bRangeEnabled", EditConditionHides, RefreshPropertyView))
 	bool bRangeScaleUniform = true;
 
 	UPROPERTY(EditInstanceOnly, Setter, Getter, DisplayName="ScaleMin", Category="Range", meta=(MotionDesignVectorWidget, AllowPreserveRatio="XYZ", Delta="0.0001", ClampMin="0", EditCondition="bRangeEnabled && !bRangeScaleUniform", EditConditionHides))
diff --git a/Engine/Plugins/VirtualProduction/ClonerEffector/Source/ClonerEffector/Public/Cloner/Extensions/CEClonerStepExtension.h b/Engine/Plugins/VirtualProduction/ClonerEffector/Source/ClonerEffector/Public/Cloner/Extensions/CEClonerStepExtension.h
--- a/Engine/Plugins/VirtualProduction/ClonerEffector/Source/ClonerEffector/Public/Cloner/Extensions/CEClonerStepExtension.h
+++ b/Engine/Plugins/VirtualProduction/ClonerEffector/Source/ClonerEffector/Public/Cloner/Extensions/CEClonerStepExtension.h
@@ -63,7 +63,7 @@ protected:
 	//~ End UCEClonerExtensionBase
 
 	/** Enable steps to add delta variation on each clone instance */
-	UPROPERTY(EditInstanceOnly, Setter="SetDeltaStepEnabled", Getter="GetDeltaStepEnabled", DisplayName="Enabled", Category="Step")
+	UPROPERTY(EditInstanceOnly, Setter="SetDeltaStepEnabled", Getter="GetDeltaStepEnabled", DisplayName="Enabled", Category="Step", meta=(RefreshPropertyView))
 	bool bDeltaStepEnabled = false;
 
 	/** Amount of position difference between one step and the next one */
diff --git a/Engine/Plugins/VirtualProduction/ClonerEffector/Source/ClonerEffector/Public/Cloner/Extensions/CEClonerTextureExtension.h b/Engine/Plugins/VirtualProduction/ClonerEffector/Source/ClonerEffector/Public/Cloner/Extensions/CEClonerTextureExtension.h
--- a/Engine/Plugins/VirtualProduction/ClonerEffector/Source/ClonerEffector/Public/Cloner/Extensions/CEClonerTextureExtension.h
+++ b/Engine/Plugins/VirtualProduction/ClonerEffector/Source/ClonerEffector/Public/Cloner/Extensions/CEClonerTextureExtension.h
@@ -120,10 +120,10 @@ protected:
 	virtual void OnExtensionDirtied(const UCEClonerExtensionBase* InExtension) override;
 	//~ End UCEClonerExtensionBase
 
-	UPROPERTY(EditInstanceOnly, Setter="SetTextureEnabled", Getter="GetTextureEnabled", DisplayName="Enabled", Category="Texture")
+	UPROPERTY(EditInstanceOnly, Setter="SetTextureEnabled", Getter="GetTextureEnabled", DisplayName="Enabled", Category="Texture", meta=(RefreshPropertyView))
 	bool bTextureEnabled = false;
 
-	UPROPERTY(EditInstanceOnly, Setter, Getter, DisplayName="Texture Provider", Category="Texture", meta=(EditCondition="bTextureEnabled", EditConditionHides))
+	UPROPERTY(EditInstanceOnly, Setter, Getter, DisplayName="Texture Provider", Category="Texture", meta=(EditCondition="bTextureEnabled", EditConditionHides, RefreshPropertyView))
 	ECEClonerTextureProvider TextureProvider = ECEClonerTextureProvider::Constraint;
 
 	UPROPERTY(EditInstanceOnly, Setter, Getter, DisplayName="Texture Asset", Category="Texture", meta=(EditCondition="bTextureEnabled && TextureProvider == ECEClonerTextureProvider::Custom", EditConditionHides))
diff --git a/Engine/Plugins/VirtualProduction/ClonerEffector/Source/ClonerEffector/Public/Cloner/Layouts/CEClonerCircleLayout.h b/Engine/Plugins/VirtualProduction/ClonerEffector/Source/ClonerEffector/Public/Cloner/Layouts/CEClonerCircleLayout.h
--- a/Engine/Plugins/VirtualProduction/ClonerEffector/Source/ClonerEffector/Public/Cloner/Layouts/CEClonerCircleLayout.h
+++ b/Engine/Plugins/VirtualProduction/ClonerEffector/Source/ClonerEffector/Public/Cloner/Layouts/CEClonerCircleLayout.h
@@ -132,7 +132,7 @@ protected:
 	UPROPERTY(EditInstanceOnly, Setter="SetOrientMesh", Getter="GetOrientMesh", Category="Layout")
 	bool bOrientMesh = true;
 
-	UPROPERTY(EditInstanceOnly, Setter, Getter, Category="Layout")
+	UPROPERTY(EditInstanceOnly, Setter, Getter, Category="Layout", meta=(RefreshPropertyView))
 	ECEClonerPlane Plane = ECEClonerPlane::XY;
 
 	UPROPERTY(EditInstanceOnly, Setter, Getter, Category="Layout", meta=(EditCondition="Plane == ECEClonerPlane::Custom", EditConditionHides))
diff --git a/Engine/Plugins/VirtualProduction/ClonerEffector/Source/ClonerEffector/Public/Cloner/Layouts/CEClonerCylinderLayout.h b/Engine/Plugins/VirtualProduction/ClonerEffector/Source/ClonerEffector/Public/Cloner/Layouts/CEClonerCylinderLayout.h
--- a/Engine/Plugins/VirtualProduction/ClonerEffector/Source/ClonerEffector/Public/Cloner/Layouts/CEClonerCylinderLayout.h
+++ b/Engine/Plugins/VirtualProduction/ClonerEffector/Source/ClonerEffector/Public/Cloner/Layouts/CEClonerCylinderLayout.h
@@ -144,7 +144,7 @@ protected:
 	UPROPERTY(EditInstanceOnly, Setter="SetOrientMesh", Getter="GetOrientMesh", Category="Layout")
 	bool bOrientMesh = true;
 
-	UPROPERTY(EditInstanceOnly, Setter, Getter, Category="Layout")
+	UPROPERTY(EditInstanceOnly, Setter, Getter, Category="Layout", meta=(RefreshPropertyView))
 	ECEClonerPlane Plane = ECEClonerPlane::XY;
 
 	UPROPERTY(EditInstanceOnly, Setter, Getter, Category="Layout", meta=(EditCondition="Plane == ECEClonerPlane::Custom", EditConditionHides))
diff --git a/Engine/Plugins/VirtualProduction/ClonerEffector/Source/ClonerEffector/Public/Cloner/Layouts/CEClonerLineLayout.h b/Engine/Plugins/VirtualProduction/ClonerEffector/Source/ClonerEffector/Public/Cloner/Layouts/CEClonerLineLayout.h
--- a/Engine/Plugins/VirtualProduction/ClonerEffector/Source/ClonerEffector/Public/Cloner/Layouts/CEClonerLineLayout.h
+++ b/Engine/Plugins/VirtualProduction/ClonerEffector/Source/ClonerEffector/Public/Cloner/Layouts/CEClonerLineLayout.h
@@ -84,7 +84,7 @@ protected:
 	UPROPERTY(EditInstanceOnly, Setter, Getter, Category="Layout")
 	float Spacing = 105.f;
 
-	UPROPERTY(EditInstanceOnly, Setter, Getter, Category="Layout")
+	UPROPERTY(EditInstanceOnly, Setter, Getter, Category="Layout", meta=(RefreshPropertyView))
 	ECEClonerAxis Axis = ECEClonerAxis::Y;
 
 	UPROPERTY(EditInstanceOnly, Setter, Getter, Category="Layout", meta=(ClampMin="0", ClampMax="1", EditCondition="Axis == ECEClonerAxis::Custom", EditConditionHides))
diff --git a/Engine/Plugins/VirtualProduction/ClonerEffector/Source/ClonerEffector/Public/Effector/Effects/CEEffectorDelayEffect.h b/Engine/Plugins/VirtualProduction/ClonerEffector/Source/ClonerEffector/Public/Effector/Effects/CEEffectorDelayEffect.h
--- a/Engine/Plugins/VirtualProduction/ClonerEffector/Source/ClonerEffector/Public/Effector/Effects/CEEffectorDelayEffect.h
+++ b/Engine/Plugins/VirtualProduction/ClonerEffector/Source/ClonerEffector/Public/Effector/Effects/CEEffectorDelayEffect.h
@@ -78,7 +78,7 @@ protected:
 	//~ End UCEEffectorEffectBase
 
 	/** Enable time delay for the mode effect */
-	UPROPERTY(EditInstanceOnly, Setter="SetDelayEnabled", Getter="GetDelayEnabled", DisplayName="Enabled", Category="Delay")
+	UPROPERTY(EditInstanceOnly, Setter="SetDelayEnabled", Getter="GetDelayEnabled", DisplayName="Enabled", Category="Delay", meta=(RefreshPropertyView))
 	bool bDelayEnabled = false;
 	
 	/** Time delay for the effect to reach its peak */
diff --git a/Engine/Plugins/VirtualProduction/ClonerEffector/Source/ClonerEffector/Public/Effector/Effects/CEEffectorForceEffect.h b/Engine/Plugins/VirtualProduction/ClonerEffector/Source/ClonerEffector/Public/Effector/Effects/CEEffectorForceEffect.h
--- a/Engine/Plugins/VirtualProduction/ClonerEffector/Source/ClonerEffector/Public/Effector/Effects/CEEffectorForceEffect.h
+++ b/Engine/Plugins/VirtualProduction/ClonerEffector/Source/ClonerEffector/Public/Effector/Effects/CEEffectorForceEffect.h
@@ -228,11 +228,11 @@ protected:
 	void OnForceOptionsChanged();
 
 	/** Forces global state */
-	UPROPERTY(EditInstanceOnly, Setter="SetForcesEnabled", Getter="GetForcesEnabled", DisplayName="Enabled", Category="Force")
+	UPROPERTY(EditInstanceOnly, Setter="SetForcesEnabled", Getter="GetForcesEnabled", DisplayName="Enabled", Category="Force", meta=(RefreshPropertyView))
 	bool bForcesEnabled = false;
 
 	/** Enable orientation force to allow each clone instance to rotate around its pivot */
-	UPROPERTY(EditInstanceOnly, Setter="SetOrientationForceEnabled", Getter="GetOrientationForceEnabled", Category="Force", meta=(EditCondition="bForcesEnabled", EditConditionHides))
+	UPROPERTY(EditInstanceOnly, Setter="SetOrientationForceEnabled", Getter="GetOrientationForceEnabled", Category="Force", meta=(EditCondition="bForcesEnabled", EditConditionHides, RefreshPropertyView))
 	bool bOrientationForceEnabled = false;
 
 	UPROPERTY(EditInstanceOnly, Setter, Getter, Category="Force", meta=(Delta="0.0001", ClampMin="0", EditCondition="bForcesEnabled && bOrientationForceEnabled", EditConditionHides))
@@ -245,7 +245,7 @@ protected:
 	FVector OrientationForceMax = FVector(0.1f);
 
 	/** Enable vortex force to allow each clone instance to rotate around a specific axis on the cloner pivot */
-	UPROPERTY(EditInstanceOnly, Setter="SetVortexForceEnabled", Getter="GetVortexForceEnabled", Category="Force", meta=(EditCondition="bForcesEnabled", EditConditionHides))
+	UPROPERTY(EditInstanceOnly, Setter="SetVortexForceEnabled", Getter="GetVortexForceEnabled", Category="Force", meta=(EditCondition="bForcesEnabled", EditConditionHides, RefreshPropertyView))
 	bool bVortexForceEnabled = false;
 
 	UPROPERTY(EditInstanceOnly, Setter, Getter, Category="Force", meta=(EditCondition="bForcesEnabled && bVortexForceEnabled", EditConditionHides))
@@ -255,7 +255,7 @@ protected:
 	FVector VortexForceAxis = FVector::ZAxisVector;
 
 	/** Enable curl noise force to allow each clone instance to add random location variation */
-	UPROPERTY(EditInstanceOnly, Setter="SetCurlNoiseForceEnabled", Getter="GetCurlNoiseForceEnabled", Category="Force", meta=(EditCondition="bForcesEnabled", EditConditionHides))
+	UPROPERTY(EditInstanceOnly, Setter="SetCurlNoiseForceEnabled", Getter="GetCurlNoiseForceEnabled", Category="Force", meta=(EditCondition="bForcesEnabled", EditConditionHides, RefreshPropertyView))
 	bool bCurlNoiseForceEnabled = false;
 
 	UPROPERTY(EditInstanceOnly, Setter, Getter, Category="Force", meta=(EditCondition="bForcesEnabled && bCurlNoiseForceEnabled", EditConditionHides))
@@ -265,7 +265,7 @@ protected:
 	float CurlNoiseForceFrequency = 10.f;
 
 	/** Enable attraction force to allow each clone instances to gravitate toward a location */
-	UPROPERTY(EditInstanceOnly, Setter="SetAttractionForceEnabled", Getter="GetAttractionForceEnabled", Category="Force", meta=(EditCondition="bForcesEnabled", EditConditionHides))
+	UPROPERTY(EditInstanceOnly, Setter="SetAttractionForceEnabled", Getter="GetAttractionForceEnabled", Category="Force", meta=(EditCondition="bForcesEnabled", EditConditionHides, RefreshPropertyView))
 	bool bAttractionForceEnabled = false;
 
 	UPROPERTY(EditInstanceOnly, Setter, Getter, Category="Force", meta=(EditCondition="bForcesEnabled && bAttractionForceEnabled", EditConditionHides))
@@ -275,14 +275,14 @@ protected:
 	float AttractionForceFalloff = 0.1f;
 
 	/** Enable gravity force to pull particles based on an acceleration vector */
-	UPROPERTY(EditInstanceOnly, Setter="SetGravityForceEnabled", Getter="GetGravityForceEnabled", Category="Force", meta=(EditCondition="bForcesEnabled", EditConditionHides))
+	UPROPERTY(EditInstanceOnly, Setter="SetGravityForceEnabled", Getter="GetGravityForceEnabled", Category="Force", meta=(EditCondition="bForcesEnabled", EditConditionHides, RefreshPropertyView))
 	bool bGravityForceEnabled = false;
 
 	UPROPERTY(EditInstanceOnly, Setter, Getter, Category="Force", meta=(EditCondition="bForcesEnabled && bGravityForceEnabled", EditConditionHides))
 	FVector GravityForceAcceleration = FVector(0, 0, -980.f);
 
 	/** Enable drag force to decrease particles velocity */
-	UPROPERTY(EditInstanceOnly, Setter="SetDragForceEnabled", Getter="GetDragForceEnabled", Category="Force", meta=(EditCondition="bForcesEnabled", EditConditionHides))
+	UPROPERTY(EditInstanceOnly, Setter="SetDragForceEnabled", Getter="GetDragForceEnabled", Category="Force", meta=(EditCondition="bForcesEnabled", EditConditionHides, RefreshPropertyView))
 	bool bDragForceEnabled = false;
 
 	UPROPERTY(EditInstanceOnly, Setter, Getter, Category="Force", meta=(EditCondition="bForcesEnabled && bDragForceEnabled", EditConditionHides))
@@ -292,7 +292,7 @@ protected:
 	float DragForceRotational = 0.f;
 
 	/** Enable vector random noise force to add variation to clones behavior */
-	UPROPERTY(EditInstanceOnly, Setter="SetVectorNoiseForceEnabled", Getter="GetVectorNoiseForceEnabled", Category="Force", meta=(EditCondition="bForcesEnabled", EditConditionHides))
+	UPROPERTY(EditInstanceOnly, Setter="SetVectorNoiseForceEnabled", Getter="GetVectorNoiseForceEnabled", Category="Force", meta=(EditCondition="bForcesEnabled", EditConditionHides, RefreshPropertyView))
 	bool bVectorNoiseForceEnabled = false;
 
 	UPROPERTY(EditInstanceOnly, Setter, Getter, Category="Force", meta=(EditCondition="bForcesEnabled && bVectorNoiseForceEnabled", EditConditionHides))
diff --git a/Engine/Plugins/VirtualProduction/ClonerEffector/Source/ClonerEffectorEditor/Private/CEEditorInputPreprocessor.cpp b/Engine/Plugins/VirtualProduction/ClonerEffector/Source/ClonerEffectorEditor/Private/CEEditorInputPreprocessor.cpp
--- a/Engine/Plugins/VirtualProduction/ClonerEffector/Source/ClonerEffectorEditor/Private/CEEditorInputPreprocessor.cpp
+++ b/Engine/Plugins/VirtualProduction/ClonerEffector/Source/ClonerEffectorEditor/Private/CEEditorInputPreprocessor.cpp
@@ -20,35 +20,29 @@ void FCEEditorInputPreprocessor::RefocusLastWidget()
 
 	if (const TSharedPtr<SWidget> FocusedWidget = SlateApp.GetKeyboardFocusedWidget())
 	{
-		if (const TSharedPtr<SWindow> SlateWindow = SlateApp.FindWidgetWindow(FocusedWidget.ToSharedRef()))
+		const FVector2D ScreenPos = FocusedWidget->GetCachedGeometry().GetAbsolutePosition() + FocusedWidget->GetCachedGeometry().GetLocalSize() / 2;
+
+		FTSTicker::GetCoreTicker().AddTicker(FTickerDelegate::CreateSPLambda(this, [ScreenPos](float)->bool
 		{
-			if (TSharedPtr<FGenericWindow> NativeWindow = SlateWindow->GetNativeWindow())
+			if (!FSlateApplication::IsInitialized())
 			{
-				const FVector2D ScreenPos = FocusedWidget->GetCachedGeometry().GetAbsolutePosition() + FocusedWidget->GetCachedGeometry().GetLocalSize() / 2;
-
-				FTSTicker::GetCoreTicker().AddTicker(FTickerDelegate::CreateSPLambda(this, [ScreenPos](float)->bool
-				{
-					if (!FSlateApplication::IsInitialized())
-					{
-						return false;
-					}
-
-					FSlateApplication& SlateApp = FSlateApplication::Get();
+				return false;
+			}
 
-					// Hit-test at the event position
-					FWidgetPath WidgetPath = SlateApp.LocateWindowUnderMouse(ScreenPos, SlateApp.GetInteractiveTopLevelWindows());
+			FSlateApplication& SlateApp = FSlateApplication::Get();
 
-					if (WidgetPath.IsValid() && WidgetPath.Widgets.Num() > 0)
-					{
-						// Last in the path is the most specific widget under the cursor
-						TSharedRef<SWidget> TargetWidget = WidgetPath.Widgets.Last().Widget;
-						SlateApp.SetAllUserFocus(TargetWidget);
-					}
+			// Hit-test at the event position
+			const FWidgetPath WidgetPath = SlateApp.LocateWindowUnderMouse(ScreenPos, SlateApp.GetInteractiveTopLevelWindows());
 
-					return false;
-				}));
+			if (WidgetPath.IsValid() && WidgetPath.Widgets.Num() > 0)
+			{
+				// Last in the path is the most specific widget under the cursor
+				TSharedRef<SWidget> TargetWidget = WidgetPath.Widgets.Last().Widget;
+				SlateApp.SetAllUserFocus(TargetWidget);
 			}
-		}
+
+			return false;
+		}));
 	}
 }
 
diff --git a/Engine/Plugins/VirtualProduction/ClonerEffector/Source/ClonerEffectorEditor/Private/Cloner/Customizations/CEEditorClonerComponentDetailCustomization.cpp b/Engine/Plugins/VirtualProduction/ClonerEffector/Source/ClonerEffectorEditor/Private/Cloner/Customizations/CEEditorClonerComponentDetailCustomization.cpp
--- a/Engine/Plugins/VirtualProduction/ClonerEffector/Source/ClonerEffectorEditor/Private/Cloner/Customizations/CEEditorClonerComponentDetailCustomization.cpp
+++ b/Engine/Plugins/VirtualProduction/ClonerEffector/Source/ClonerEffectorEditor/Private/Cloner/Customizations/CEEditorClonerComponentDetailCustomization.cpp
@@ -298,7 +298,12 @@ void FCEEditorClonerComponentDetailCustomization::OnChildPropertyChanged(const F
 		return;
 	}
 
-	if (InEvent.ChangeType == EPropertyChangeType::Interactive)
+	if (InEvent.ChangeType == EPropertyChangeType::Interactive || !InEvent.MemberProperty)
+	{
+		return;
+	}
+
+	if (!InEvent.MemberProperty->HasMetaData(TEXT("RefreshPropertyView")))
 	{
 		return;
 	}
diff --git a/Engine/Plugins/VirtualProduction/ClonerEffector/Source/ClonerEffectorEditor/Private/Effector/Customizations/CEEditorEffectorComponentDetailCustomization.cpp b/Engine/Plugins/VirtualProduction/ClonerEffector/Source/ClonerEffectorEditor/Private/Effector/Customizations/CEEditorEffectorComponentDetailCustomization.cpp
--- a/Engine/Plugins/VirtualProduction/ClonerEffector/Source/ClonerEffectorEditor/Private/Effector/Customizations/CEEditorEffectorComponentDetailCustomization.cpp
+++ b/Engine/Plugins/VirtualProduction/ClonerEffector/Source/ClonerEffectorEditor/Private/Effector/Customizations/CEEditorEffectorComponentDetailCustomization.cpp
@@ -167,7 +167,12 @@ void FCEEditorEffectorComponentDetailCustomization::OnChildPropertyChanged(const
 		return;
 	}
 
-	if (InEvent.ChangeType == EPropertyChangeType::Interactive)
+	if (InEvent.ChangeType == EPropertyChangeType::Interactive || !InEvent.MemberProperty)
+	{
+		return;
+	}
+
+	if (!InEvent.MemberProperty->HasMetaData(TEXT("RefreshPropertyView")))
 	{
 		return;
 	}
