MotionDesign : ClonerEffector
- Fixed skipping materials causing simulation to not be up to date with new materials

diff --git a/Engine/Plugins/VirtualProduction/ClonerEffector/Source/ClonerEffector/Private/Cloner/Attachments/CEClonerAttachmentTreeBehavior.cpp b/Engine/Plugins/VirtualProduction/ClonerEffector/Source/ClonerEffector/Private/Cloner/Attachments/CEClonerAttachmentTreeBehavior.cpp
--- a/Engine/Plugins/VirtualProduction/ClonerEffector/Source/ClonerEffector/Private/Cloner/Attachments/CEClonerAttachmentTreeBehavior.cpp
+++ b/Engine/Plugins/VirtualProduction/ClonerEffector/Source/ClonerEffector/Private/Cloner/Attachments/CEClonerAttachmentTreeBehavior.cpp
@@ -606,9 +606,10 @@ void FCEClonerAttachmentGroupBehavior::RefreshBakedMeshRootMaterials(FCEClonerAt
 					StaticMaterials[MaterialSlotIndex].MaterialInterface = Material;
 				}
 			}
-
-			MaterialIndexOffset++;
 		}
+
+		// Increment offset by material count for next attachment
+		MaterialIndexOffset += AttachmentItem->BakedMaterials.Num();
 	}
 
 	InTree.bItemAttachmentsDirty = true;
