Motion Design: corrected log verbosity from Log to Verbose for logs that were constantly outputting for any valid scene and removed redundant logs.

#Synced-CL 48355279

diff --git a/Engine/Plugins/VirtualProduction/ActorModifierCore/Source/ActorModifierCore/Private/Subsystems/ActorModifierCoreSubsystem.cpp b/Engine/Plugins/VirtualProduction/ActorModifierCore/Source/ActorModifierCore/Private/Subsystems/ActorModifierCoreSubsystem.cpp
--- a/Engine/Plugins/VirtualProduction/ActorModifierCore/Source/ActorModifierCore/Private/Subsystems/ActorModifierCoreSubsystem.cpp
+++ b/Engine/Plugins/VirtualProduction/ActorModifierCore/Source/ActorModifierCore/Private/Subsystems/ActorModifierCoreSubsystem.cpp
@@ -1481,7 +1481,7 @@ bool UActorModifierCoreSubsystem::RegisterModifierSharedProvider(AActorModifierC
 	UActorModifierCoreSubsystem* MutableThis = const_cast<UActorModifierCoreSubsystem*>(this);
 	MutableThis->ModifierSharedProviders.Add(Level, InSharedActor);
 
-	UE_LOG(LogActorModifierCoreSubsystem, Log, TEXT("Modifier shared provider registered for world %s and level %s"), *Level->GetWorld()->GetDebugDisplayName(), *Level->GetName());
+	UE_LOG(LogActorModifierCoreSubsystem, Verbose, TEXT("Modifier shared provider registered for world %s and level %s"), *Level->GetWorld()->GetDebugDisplayName(), *Level->GetName());
 
 	return true;
 }
@@ -1510,7 +1510,7 @@ bool UActorModifierCoreSubsystem::UnregisterModifierSharedProvider(const AActor*
 		return false;
 	}
 
-	UE_LOG(LogActorModifierCoreSubsystem, Log, TEXT("Modifier shared provider unregistered for world %s and level %s"), *Level->GetWorld()->GetDebugDisplayName(), *Level->GetName());
+	UE_LOG(LogActorModifierCoreSubsystem, Verbose, TEXT("Modifier shared provider unregistered for world %s and level %s"), *Level->GetWorld()->GetDebugDisplayName(), *Level->GetName());
 
 	UActorModifierCoreSubsystem* MutableThis = const_cast<UActorModifierCoreSubsystem*>(this);
 	return MutableThis->ModifierSharedProviders.Remove(Level) > 0;
diff --git a/Engine/Plugins/VirtualProduction/ClonerEffector/Source/ClonerEffector/Private/Cloner/Attachments/CEClonerAttachmentTree.cpp b/Engine/Plugins/VirtualProduction/ClonerEffector/Source/ClonerEffector/Private/Cloner/Attachments/CEClonerAttachmentTree.cpp
--- a/Engine/Plugins/VirtualProduction/ClonerEffector/Source/ClonerEffector/Private/Cloner/Attachments/CEClonerAttachmentTree.cpp
+++ b/Engine/Plugins/VirtualProduction/ClonerEffector/Source/ClonerEffector/Private/Cloner/Attachments/CEClonerAttachmentTree.cpp
@@ -95,7 +95,7 @@ bool FCEClonerAttachmentTree::UpdateAttachments(bool bInReset)
 		return false;
 	}
 
-	UE_LOG(LogCECloner, Log, TEXT("%s : updating attachment tree"), *AttachmentRoot->GetActorNameOrLabel())
+	UE_LOG(LogCECloner, Verbose, TEXT("%s : updating attachment tree"), *AttachmentRoot->GetActorNameOrLabel())
 
 	Status = ECEClonerAttachmentStatus::Updating;
 	const bool bResult = UpdateAttachmentsInternal();
diff --git a/Engine/Plugins/VirtualProduction/ClonerEffector/Source/ClonerEffector/Private/Cloner/Attachments/CEClonerAttachmentTreeBehavior.cpp b/Engine/Plugins/VirtualProduction/ClonerEffector/Source/ClonerEffector/Private/Cloner/Attachments/CEClonerAttachmentTreeBehavior.cpp
--- a/Engine/Plugins/VirtualProduction/ClonerEffector/Source/ClonerEffector/Private/Cloner/Attachments/CEClonerAttachmentTreeBehavior.cpp
+++ b/Engine/Plugins/VirtualProduction/ClonerEffector/Source/ClonerEffector/Private/Cloner/Attachments/CEClonerAttachmentTreeBehavior.cpp
@@ -420,7 +420,7 @@ void FCEClonerAttachmentGroupBehavior::OnComponentTransformed(FCEClonerAttachmen
 		return;
 	}
 
-	UE_LOG(LogCECloner, Log, TEXT("%s : Transform state changed for %s"), *ClonerActor->GetActorNameOrLabel(), *Owner->GetActorNameOrLabel());
+	UE_LOG(LogCECloner, Verbose, TEXT("%s : Transform state changed for %s"), *ClonerActor->GetActorNameOrLabel(), *Owner->GetActorNameOrLabel());
 
 	Item->MeshStatus = ECEClonerAttachmentStatus::Outdated;
 	InvalidateMesh(InTree, Owner);
@@ -458,7 +458,7 @@ void FCEClonerAttachmentGroupBehavior::OnMeshChanged(FCEClonerAttachmentTree& In
 		return;
 	}
 
-	UE_LOG(LogCECloner, Log, TEXT("%s : Detected mesh change for %s"), *ClonerActor->GetActorNameOrLabel(), *InActor->GetActorNameOrLabel());
+	UE_LOG(LogCECloner, Verbose, TEXT("%s : Detected mesh change for %s"), *ClonerActor->GetActorNameOrLabel(), *InActor->GetActorNameOrLabel());
 
 	Item->MeshStatus = ECEClonerAttachmentStatus::Outdated;
 	InvalidateMesh(InTree, InActor);
@@ -489,7 +489,7 @@ void FCEClonerAttachmentGroupBehavior::OnRenderStateDirty(FCEClonerAttachmentTre
 		return;
 	}
 
-	UE_LOG(LogCECloner, Log, TEXT("%s : Render state changed for %s"), *ClonerActor->GetActorNameOrLabel(), *Owner->GetActorNameOrLabel());
+	UE_LOG(LogCECloner, Verbose, TEXT("%s : Render state changed for %s"), *ClonerActor->GetActorNameOrLabel(), *Owner->GetActorNameOrLabel());
 
 	// Rebind delegates as new components might be available
 	BindActorDelegates(InTree, Owner);
@@ -555,7 +555,7 @@ void FCEClonerAttachmentGroupBehavior::OnCheckMaterialChanged(FCEClonerAttachmen
 			}
 		}
 
-		UE_LOG(LogCECloner, Log, TEXT("%s : Detected material change for %s"), *ClonerActor->GetActorNameOrLabel(), *ActorChanged->GetActorNameOrLabel());
+		UE_LOG(LogCECloner, Verbose, TEXT("%s : Detected material change for %s"), *ClonerActor->GetActorNameOrLabel(), *ActorChanged->GetActorNameOrLabel());
 
 		if (AttachmentItem->MeshStatus == ECEClonerAttachmentStatus::Outdated)
 		{
diff --git a/Engine/Plugins/VirtualProduction/ClonerEffector/Source/ClonerEffector/Private/Cloner/CEClonerComponent.cpp b/Engine/Plugins/VirtualProduction/ClonerEffector/Source/ClonerEffector/Private/Cloner/CEClonerComponent.cpp
--- a/Engine/Plugins/VirtualProduction/ClonerEffector/Source/ClonerEffector/Private/Cloner/CEClonerComponent.cpp
+++ b/Engine/Plugins/VirtualProduction/ClonerEffector/Source/ClonerEffector/Private/Cloner/CEClonerComponent.cpp
@@ -872,8 +872,6 @@ void UCEClonerComponent::ConvertToStaticMesh()
 	FScopedSlowTask SlowTask(0.0f, LOCTEXT("ConvertToStaticMesh", "Converting cloner to static mesh"));
 	SlowTask.MakeDialog();
 
-	UE_LOG(LogCECloner, Log, TEXT("%s : Request ConvertToStaticMesh..."), *GetOwner()->GetActorNameOrLabel())
-
 	if (UE::ClonerEffector::Conversion::ConvertClonerToStaticMesh(this))
 	{
 		UE_LOG(LogCECloner, Log, TEXT("%s : ConvertToStaticMesh Completed"), *GetOwner()->GetActorNameOrLabel())
@@ -894,8 +892,6 @@ void UCEClonerComponent::ConvertToDynamicMesh()
 	FScopedSlowTask SlowTask(0.0f, LOCTEXT("ConvertToDynamicMesh", "Converting cloner to dynamic mesh"));
 	SlowTask.MakeDialog();
 
-	UE_LOG(LogCECloner, Log, TEXT("%s : Request ConvertToDynamicMesh..."), *GetOwner()->GetActorNameOrLabel())
-
 	if (UE::ClonerEffector::Conversion::ConvertClonerToDynamicMesh(this))
 	{
 		UE_LOG(LogCECloner, Log, TEXT("%s : ConvertToDynamicMesh Completed"), *GetOwner()->GetActorNameOrLabel())
@@ -916,8 +912,6 @@ void UCEClonerComponent::ConvertToStaticMeshes()
 	FScopedSlowTask SlowTask(0.0f, LOCTEXT("ConvertToStaticMeshes", "Converting cloner to static meshes"));
 	SlowTask.MakeDialog();
 
-	UE_LOG(LogCECloner, Log, TEXT("%s : Request ConvertToStaticMeshes..."), *GetOwner()->GetActorNameOrLabel())
-
 	if (!UE::ClonerEffector::Conversion::ConvertClonerToStaticMeshes(this).IsEmpty())
 	{
 		UE_LOG(LogCECloner, Log, TEXT("%s : ConvertToStaticMeshes Completed"), *GetOwner()->GetActorNameOrLabel())
@@ -938,8 +932,6 @@ void UCEClonerComponent::ConvertToDynamicMeshes()
 	FScopedSlowTask SlowTask(0.0f, LOCTEXT("ConvertToDynamicMeshes", "Converting cloner to dynamic meshes"));
 	SlowTask.MakeDialog();
 
-	UE_LOG(LogCECloner, Log, TEXT("%s : Request ConvertToDynamicMeshes..."), *GetOwner()->GetActorNameOrLabel())
-
 	if (!UE::ClonerEffector::Conversion::ConvertClonerToDynamicMeshes(this).IsEmpty())
 	{
 		UE_LOG(LogCECloner, Log, TEXT("%s : ConvertToDynamicMeshes Completed"), *GetOwner()->GetActorNameOrLabel())
@@ -960,8 +952,6 @@ void UCEClonerComponent::ConvertToInstancedStaticMeshes()
 	FScopedSlowTask SlowTask(0.0f, LOCTEXT("ConvertToInstancedStaticMeshes", "Converting cloner to instanced static meshes"));
 	SlowTask.MakeDialog();
 
-	UE_LOG(LogCECloner, Log, TEXT("%s : Request ConvertToInstancedStaticMeshes..."), *GetOwner()->GetActorNameOrLabel())
-
 	if (!UE::ClonerEffector::Conversion::ConvertClonerToInstancedStaticMeshes(this).IsEmpty())
 	{
 		UE_LOG(LogCECloner, Log, TEXT("%s : ConvertToInstancedStaticMeshes Completed"), *GetOwner()->GetActorNameOrLabel())
@@ -1339,7 +1329,7 @@ void UCEClonerComponent::OnActiveLayoutLoaded(UCEClonerLayoutBase* InLayout, boo
 		return;
 	}
 
-	UE_LOG(LogCECloner, Log, TEXT("%s : Cloner layout system loaded %s - %s"), *GetOwner()->GetActorNameOrLabel(), *InLayout->GetLayoutName().ToString(), *InLayout->GetLayoutAssetPath());
+	UE_LOG(LogCECloner, Verbose, TEXT("%s : Cloner layout system loaded %s - %s"), *GetOwner()->GetActorNameOrLabel(), *InLayout->GetLayoutName().ToString(), *InLayout->GetLayoutAssetPath());
 
 	OnClonerLayoutLoadedDelegate.Broadcast(this, InLayout);
 
@@ -1371,7 +1361,7 @@ void UCEClonerComponent::ActivateLayout(UCEClonerLayoutBase* InLayout)
 
 	ActiveLayout = InLayout;
 
-	UE_LOG(LogCECloner, Log, TEXT("%s : Cloner layout system changed %s - %s"), *GetOwner()->GetActorNameOrLabel(), *InLayout->GetLayoutName().ToString(), *InLayout->GetLayoutAssetPath());
+	UE_LOG(LogCECloner, Verbose, TEXT("%s : Cloner layout system changed %s - %s"), *GetOwner()->GetActorNameOrLabel(), *InLayout->GetLayoutName().ToString(), *InLayout->GetLayoutAssetPath());
 
 	OnActiveLayoutChanged();
 
@@ -1502,7 +1492,7 @@ void UCEClonerComponent::UpdateClonerMeshes()
 
 	if (bMeshChanged)
 	{
-		UE_LOG(LogCECloner, Log, TEXT("%s : Cloner mesh updated - %i cached meshes - %i rendered meshes"), *ClonerActor->GetActorNameOrLabel(), ClonerTree.RootItemAttachments.Num(), MeshRenderer->Meshes.Num())
+		UE_LOG(LogCECloner, Verbose, TEXT("%s : Cloner mesh updated - %i cached meshes - %i rendered meshes"), *ClonerActor->GetActorNameOrLabel(), ClonerTree.RootItemAttachments.Num(), MeshRenderer->Meshes.Num())
 	}
 
 	// Set new number of meshes in renderer
diff --git a/Engine/Plugins/VirtualProduction/ClonerEffector/Source/ClonerEffector/Private/Cloner/Extensions/CEClonerEffectorExtension.cpp b/Engine/Plugins/VirtualProduction/ClonerEffector/Source/ClonerEffector/Private/Cloner/Extensions/CEClonerEffectorExtension.cpp
--- a/Engine/Plugins/VirtualProduction/ClonerEffector/Source/ClonerEffector/Private/Cloner/Extensions/CEClonerEffectorExtension.cpp
+++ b/Engine/Plugins/VirtualProduction/ClonerEffector/Source/ClonerEffector/Private/Cloner/Extensions/CEClonerEffectorExtension.cpp
@@ -41,7 +41,7 @@ bool UCEClonerEffectorExtension::LinkEffector(AActor* InEffectorActor)
 
 	OnEffectorsChanged();
 
-	UE_LOG(LogCECloner, Log, TEXT("%s : Effector %s linked to Cloner"), *GetClonerComponent()->GetOwner()->GetActorNameOrLabel(), *InEffectorActor->GetActorNameOrLabel());
+	UE_LOG(LogCECloner, Verbose, TEXT("%s : Effector %s linked to Cloner"), *GetClonerComponent()->GetOwner()->GetActorNameOrLabel(), *InEffectorActor->GetActorNameOrLabel());
 
 	return true;
 }
@@ -57,7 +57,7 @@ bool UCEClonerEffectorExtension::UnlinkEffector(AActor* InEffectorActor)
 	{
 		OnEffectorsChanged();
 
-		UE_LOG(LogCECloner, Log, TEXT("%s : Effector %s unlinked from Cloner"), *GetClonerComponent()->GetOwner()->GetActorNameOrLabel(), *InEffectorActor->GetActorNameOrLabel());
+		UE_LOG(LogCECloner, Verbose, TEXT("%s : Effector %s unlinked from Cloner"), *GetClonerComponent()->GetOwner()->GetActorNameOrLabel(), *InEffectorActor->GetActorNameOrLabel());
 	}
 
 	return true;
diff --git a/Engine/Plugins/VirtualProduction/ClonerEffector/Source/ClonerEffector/Private/Cloner/Layouts/CEClonerLayoutBase.cpp b/Engine/Plugins/VirtualProduction/ClonerEffector/Source/ClonerEffector/Private/Cloner/Layouts/CEClonerLayoutBase.cpp
--- a/Engine/Plugins/VirtualProduction/ClonerEffector/Source/ClonerEffector/Private/Cloner/Layouts/CEClonerLayoutBase.cpp
+++ b/Engine/Plugins/VirtualProduction/ClonerEffector/Source/ClonerEffector/Private/Cloner/Layouts/CEClonerLayoutBase.cpp
@@ -464,7 +464,7 @@ void UCEClonerLayoutBase::OnLevelCleanup()
 {
 	if (IsLayoutLoaded())
 	{
-		UE_LOG(LogCECloner, Log, TEXT("%s : Cloner layout cleanup %s"), *GetClonerActor()->GetActorNameOrLabel(), *LayoutName.ToString())
+		UE_LOG(LogCECloner, Verbose, TEXT("%s : Cloner layout cleanup %s"), *GetClonerActor()->GetActorNameOrLabel(), *LayoutName.ToString())
 
 		DeactivateLayout();
 
diff --git a/Engine/Plugins/VirtualProduction/ClonerEffector/Source/ClonerEffector/Private/Subsystems/CEEffectorSubsystem.cpp b/Engine/Plugins/VirtualProduction/ClonerEffector/Source/ClonerEffector/Private/Subsystems/CEEffectorSubsystem.cpp
--- a/Engine/Plugins/VirtualProduction/ClonerEffector/Source/ClonerEffector/Private/Subsystems/CEEffectorSubsystem.cpp
+++ b/Engine/Plugins/VirtualProduction/ClonerEffector/Source/ClonerEffector/Private/Subsystems/CEEffectorSubsystem.cpp
@@ -98,7 +98,7 @@ bool UCEEffectorSubsystem::RegisterChannelEffector(UCEEffectorComponent* InEffec
 	if (EffectorIndex == INDEX_NONE)
 	{
 		EffectorIndex = EffectorsWeak.Add(InEffector);
-		UE_LOG(LogCEEffector, Log, TEXT("%s effector registered in channel %i"), *InEffector->GetOwner()->GetActorNameOrLabel(), EffectorIndex);
+		UE_LOG(LogCEEffector, Verbose, TEXT("%s effector registered in channel %i"), *InEffector->GetOwner()->GetActorNameOrLabel(), EffectorIndex);
 	}
 
 	if (InEffector->GetChannelData().Identifier != EffectorIndex)
@@ -122,7 +122,7 @@ bool UCEEffectorSubsystem::UnregisterChannelEffector(UCEEffectorComponent* InEff
 
 	if (bUnregistered)
 	{
-		UE_LOG(LogCEEffector, Log, TEXT("%s effector unregistered from channel"), *InEffector->GetOwner()->GetActorNameOrLabel());
+		UE_LOG(LogCEEffector, Verbose, TEXT("%s effector unregistered from channel"), *InEffector->GetOwner()->GetActorNameOrLabel());
 
 		const int32 OldIdentifier = InEffector->GetChannelData().Identifier;
 		InEffector->GetChannelData().Identifier = INDEX_NONE;
diff --git a/Engine/Plugins/VirtualProduction/PropertyAnimatorCore/Source/PropertyAnimatorCore/Private/Animators/PropertyAnimatorCoreBase.cpp b/Engine/Plugins/VirtualProduction/PropertyAnimatorCore/Source/PropertyAnimatorCore/Private/Animators/PropertyAnimatorCoreBase.cpp
--- a/Engine/Plugins/VirtualProduction/PropertyAnimatorCore/Source/PropertyAnimatorCore/Private/Animators/PropertyAnimatorCoreBase.cpp
+++ b/Engine/Plugins/VirtualProduction/PropertyAnimatorCore/Source/PropertyAnimatorCore/Private/Animators/PropertyAnimatorCoreBase.cpp
@@ -496,7 +496,7 @@ void UPropertyAnimatorCoreBase::OnAnimatorRemoved(EPropertyAnimatorCoreUpdateEve
 void UPropertyAnimatorCoreBase::OnAnimatorEnabled(EPropertyAnimatorCoreUpdateEvent InType)
 {
 	UE_LOG(LogPropertyAnimatorCore
-		, Log
+		, Verbose
 		, TEXT("%s : PropertyAnimator %s (%s) enabled %i")
 		, GetAnimatorActor() ? *GetAnimatorActor()->GetActorNameOrLabel() : TEXT("Invalid Actor")
 		, *GetAnimatorDisplayName().ToString()
@@ -507,7 +507,7 @@ void UPropertyAnimatorCoreBase::OnAnimatorEnabled(EPropertyAnimatorCoreUpdateEve
 void UPropertyAnimatorCoreBase::OnAnimatorDisabled(EPropertyAnimatorCoreUpdateEvent InType)
 {
 	UE_LOG(LogPropertyAnimatorCore
-		, Log
+		, Verbose
 		, TEXT("%s : PropertyAnimator %s (%s) disabled %i")
 		, GetAnimatorActor() ? *GetAnimatorActor()->GetActorNameOrLabel() : TEXT("Invalid Actor")
 		, *GetAnimatorDisplayName().ToString()
diff --git a/Engine/Plugins/VirtualProduction/RemoteControl/Source/RemoteControlLogic/Private/Action/Path/RCSetAssetByPathActionNew.cpp b/Engine/Plugins/VirtualProduction/RemoteControl/Source/RemoteControlLogic/Private/Action/Path/RCSetAssetByPathActionNew.cpp
--- a/Engine/Plugins/VirtualProduction/RemoteControl/Source/RemoteControlLogic/Private/Action/Path/RCSetAssetByPathActionNew.cpp
+++ b/Engine/Plugins/VirtualProduction/RemoteControl/Source/RemoteControlLogic/Private/Action/Path/RCSetAssetByPathActionNew.cpp
@@ -242,7 +242,7 @@ bool URCSetAssetByPathActionNew::SetInternalAsset(UObject* InSetterObject)
 			{
 				World->DestroyActor(OldActor);
 
-				UE_LOG(LogRemoteControl, Display, TEXT("Path Behavior sets Blueprint Asset %s"), *InSetterObject->GetName());
+				UE_LOG(LogRemoteControl, Verbose, TEXT("Path Behavior sets Blueprint Asset %s"), *InSetterObject->GetName());
 				return true;
 			}
 
@@ -272,7 +272,7 @@ bool URCSetAssetByPathActionNew::SetInternalAsset(UObject* InSetterObject)
 						{
 							if (ObjectProperty->PropertyClass && (SetterObjectClass == ObjectProperty->PropertyClass || SetterObjectClass->IsChildOf(ObjectProperty->PropertyClass)))
 							{
-								UE_LOG(LogRemoteControl, Display, TEXT("Path Behavior sets Array Property %s"), *InSetterObject->GetName());
+								UE_LOG(LogRemoteControl, Verbose, TEXT("Path Behavior sets Array Property %s"), *InSetterObject->GetName());
 								SubPropertyHandle->SetValue(InSetterObject);
 								return true;
 							}
@@ -285,7 +285,7 @@ bool URCSetAssetByPathActionNew::SetInternalAsset(UObject* InSetterObject)
 		{
 			if (ObjectProperty->PropertyClass && (SetterObjectClass == ObjectProperty->PropertyClass || SetterObjectClass->IsChildOf(ObjectProperty->PropertyClass)))
 			{
-				UE_LOG(LogRemoteControl, Display, TEXT("Path Behavior sets Property %s"), *InSetterObject->GetName());
+				UE_LOG(LogRemoteControl, Verbose, TEXT("Path Behavior sets Property %s"), *InSetterObject->GetName());
 				PropertyHandle->SetValue(InSetterObject);
 				return true;
 			}
diff --git a/Engine/Plugins/VirtualProduction/RemoteControl/Source/RemoteControlLogic/Private/Behaviour/Builtin/Path/RCSetAssetByPathBehaviour.cpp b/Engine/Plugins/VirtualProduction/RemoteControl/Source/RemoteControlLogic/Private/Behaviour/Builtin/Path/RCSetAssetByPathBehaviour.cpp
--- a/Engine/Plugins/VirtualProduction/RemoteControl/Source/RemoteControlLogic/Private/Behaviour/Builtin/Path/RCSetAssetByPathBehaviour.cpp
+++ b/Engine/Plugins/VirtualProduction/RemoteControl/Source/RemoteControlLogic/Private/Behaviour/Builtin/Path/RCSetAssetByPathBehaviour.cpp
@@ -245,7 +245,7 @@ bool URCSetAssetByPathBehaviour::SetInternalAsset(UObject* SetterObject)
 			UStaticMeshComponent* StaticMeshComponent = Cast<UStaticMeshComponent>(ExposedProperty->GetBoundObject());
 			StaticMeshComponent->SetStaticMesh(Cast<UStaticMesh>(SetterObject));
 			
-			UE_LOG(LogRemoteControl, Display, TEXT("Path Behaviour sets Static Mesh Asset %s"), *SetterObject->GetName());
+			UE_LOG(LogRemoteControl, Verbose, TEXT("Path Behaviour sets Static Mesh Asset %s"), *SetterObject->GetName());
 			return true;
 		}
 		
@@ -256,7 +256,7 @@ bool URCSetAssetByPathBehaviour::SetInternalAsset(UObject* SetterObject)
 				MaterialOwnerComponent->SetMaterial(0, Cast<UMaterial>(SetterObject));
 			}
 			
-			UE_LOG(LogRemoteControl, Display, TEXT("Path Behaviour sets Material Asset %s"), *SetterObject->GetName());
+			UE_LOG(LogRemoteControl, Verbose, TEXT("Path Behaviour sets Material Asset %s"), *SetterObject->GetName());
 			return true;
 		}
 
@@ -296,7 +296,7 @@ bool URCSetAssetByPathBehaviour::SetInternalAsset(UObject* SetterObject)
 			{
 				World->DestroyActor(OldActor);
 				
-				UE_LOG(LogRemoteControl, Display, TEXT("Path Behaviour sets Blueprint Asset %s"), *SetterObject->GetName());
+				UE_LOG(LogRemoteControl, Verbose, TEXT("Path Behaviour sets Blueprint Asset %s"), *SetterObject->GetName());
 				return true;
 			}
 
diff --git a/Engine/Plugins/VirtualProduction/Text3D/Source/Text3D/Private/Text3DComponent.cpp b/Engine/Plugins/VirtualProduction/Text3D/Source/Text3D/Private/Text3DComponent.cpp
--- a/Engine/Plugins/VirtualProduction/Text3D/Source/Text3D/Private/Text3DComponent.cpp
+++ b/Engine/Plugins/VirtualProduction/Text3D/Source/Text3D/Private/Text3DComponent.cpp
@@ -1079,7 +1079,7 @@ void UText3DComponent::Rebuild(EText3DRendererFlags InUpdateFlags)
 	{
 		if (!EnumHasAllFlags(UpdateFlags, InUpdateFlags))
 		{
-			UE_LOG(LogText3D, Log, TEXT("[%s] Rebuild called with update flags (%u) while component is not registered yet. This new Rebuild will be deferred to the next tick."), *GetReadableName(), InUpdateFlags);
+			UE_LOG(LogText3D, Verbose, TEXT("[%s] Rebuild called with update flags (%u) while component is not registered yet. This new Rebuild will be deferred to the next tick."), *GetReadableName(), InUpdateFlags);
 			QueueRebuild(InUpdateFlags);
 		}
 
@@ -1090,7 +1090,7 @@ void UText3DComponent::Rebuild(EText3DRendererFlags InUpdateFlags)
 	{
 		if (!EnumHasAllFlags(UpdateFlags, InUpdateFlags))
 		{
-			UE_LOG(LogText3D, Log, TEXT("[%s] Rebuild called with update flags (%u) while package is being saved. This new Rebuild will be deferred to the next tick."), *GetReadableName(), InUpdateFlags);
+			UE_LOG(LogText3D, Verbose, TEXT("[%s] Rebuild called with update flags (%u) while package is being saved. This new Rebuild will be deferred to the next tick."), *GetReadableName(), InUpdateFlags);
 			QueueRebuild(InUpdateFlags);
 		}
 
