Motion Design: moved template cloner to be a property to keep it loaded instead of loading it in play time synchronously.

#Synced-CL 48481255

diff --git a/Engine/Plugins/VirtualProduction/ClonerEffector/Source/ClonerEffector/Private/Cloner/Extensions/CEClonerCollisionExtension.cpp b/Engine/Plugins/VirtualProduction/ClonerEffector/Source/ClonerEffector/Private/Cloner/Extensions/CEClonerCollisionExtension.cpp
--- a/Engine/Plugins/VirtualProduction/ClonerEffector/Source/ClonerEffector/Private/Cloner/Extensions/CEClonerCollisionExtension.cpp
+++ b/Engine/Plugins/VirtualProduction/ClonerEffector/Source/ClonerEffector/Private/Cloner/Extensions/CEClonerCollisionExtension.cpp
@@ -4,6 +4,7 @@
 
 #include "Cloner/CEClonerComponent.h"
 #include "Components/StaticMeshComponent.h"
+#include "Engine/World.h"
 #include "Materials/MaterialInstanceDynamic.h"
 #include "NiagaraDataInterfaceArrayFloat.h"
 #include "NiagaraMeshRendererProperties.h"
@@ -182,7 +183,7 @@ void UCEClonerCollisionExtension::OnExtensionParametersChanged(UCEClonerComponen
 	}
 
 #if WITH_EDITOR
-	if (InitVisualizerComponent())
+	if (ShouldInitVisualizerComponent(InComponent) && InitVisualizerComponent())
 	{
 		CollisionVisualizerComponent->SetVisibility(bParticleCollisionEnabled && bPreviewCollisionGrid);
 		// Divide by 100 because cube size is 100x100x100
@@ -209,6 +210,13 @@ void UCEClonerCollisionExtension::OnExtensionDeactivated()
 }
 
 #if WITH_EDITOR
+
+bool UCEClonerCollisionExtension::ShouldInitVisualizerComponent(TNotNull<const UCEClonerComponent*> InComponent) const
+{
+	const UWorld* World = InComponent->GetWorld();
+	return World && (World->WorldType == EWorldType::Editor || World->WorldType == EWorldType::EditorPreview);
+}
+
 bool UCEClonerCollisionExtension::InitVisualizerComponent()
 {
 	if (CollisionVisualizerComponent)
diff --git a/Engine/Plugins/VirtualProduction/ClonerEffector/Source/ClonerEffector/Private/Cloner/Layouts/CEClonerLayoutBase.cpp b/Engine/Plugins/VirtualProduction/ClonerEffector/Source/ClonerEffector/Private/Cloner/Layouts/CEClonerLayoutBase.cpp
--- a/Engine/Plugins/VirtualProduction/ClonerEffector/Source/ClonerEffector/Private/Cloner/Layouts/CEClonerLayoutBase.cpp
+++ b/Engine/Plugins/VirtualProduction/ClonerEffector/Source/ClonerEffector/Private/Cloner/Layouts/CEClonerLayoutBase.cpp
@@ -26,10 +26,16 @@ bool UCEClonerLayoutBase::IsLayoutValid() const
 	}
 
 	// Get the template niagara asset
-	const UNiagaraSystem* TemplateNiagaraSystem = LoadObject<UNiagaraSystem>(nullptr, *LayoutAssetPath);
+	if (!TemplateNiagaraSystem)
+	{
+		TemplateNiagaraSystem = LoadObject<UNiagaraSystem>(nullptr, *LayoutAssetPath);
+	}
 
 	// Get the base niagara asset
-	const UNiagaraSystem* BaseNiagaraSystem = LoadObject<UNiagaraSystem>(nullptr, LayoutBaseAssetPath);
+	if (!BaseNiagaraSystem)
+	{
+		BaseNiagaraSystem = LoadObject<UNiagaraSystem>(nullptr, LayoutBaseAssetPath);
+	}
 
 	if (!TemplateNiagaraSystem || !BaseNiagaraSystem)
 	{
@@ -494,7 +500,12 @@ FString UCEClonerLayoutBase::GetLayoutHash() const
 {
 	FString LayoutHash;
 	
-	if (const UNiagaraSystem* TemplateNiagaraSystem = LoadObject<UNiagaraSystem>(nullptr, *LayoutAssetPath))
+	if (!TemplateNiagaraSystem)
+	{
+		TemplateNiagaraSystem = LoadObject<UNiagaraSystem>(nullptr, *LayoutAssetPath);
+	}
+
+	if (TemplateNiagaraSystem)
 	{
 		if (UPackage* Package = TemplateNiagaraSystem->GetPackage())
 		{ 
diff --git a/Engine/Plugins/VirtualProduction/ClonerEffector/Source/ClonerEffector/Public/Cloner/Extensions/CEClonerCollisionExtension.h b/Engine/Plugins/VirtualProduction/ClonerEffector/Source/ClonerEffector/Public/Cloner/Extensions/CEClonerCollisionExtension.h
--- a/Engine/Plugins/VirtualProduction/ClonerEffector/Source/ClonerEffector/Public/Cloner/Extensions/CEClonerCollisionExtension.h
+++ b/Engine/Plugins/VirtualProduction/ClonerEffector/Source/ClonerEffector/Public/Cloner/Extensions/CEClonerCollisionExtension.h
@@ -120,6 +120,7 @@ protected:
 	//~ End UCEClonerExtensionBase
 
 #if WITH_EDITOR
+	bool ShouldInitVisualizerComponent(TNotNull<const UCEClonerComponent*> InComponent) const;
 	bool InitVisualizerComponent();
 	void DestroyVisualizerComponent();
 #endif
diff --git a/Engine/Plugins/VirtualProduction/ClonerEffector/Source/ClonerEffector/Public/Cloner/Layouts/CEClonerLayoutBase.h b/Engine/Plugins/VirtualProduction/ClonerEffector/Source/ClonerEffector/Public/Cloner/Layouts/CEClonerLayoutBase.h
--- a/Engine/Plugins/VirtualProduction/ClonerEffector/Source/ClonerEffector/Public/Cloner/Layouts/CEClonerLayoutBase.h
+++ b/Engine/Plugins/VirtualProduction/ClonerEffector/Source/ClonerEffector/Public/Cloner/Layouts/CEClonerLayoutBase.h
@@ -178,6 +178,14 @@ private:
 	UPROPERTY(DuplicateTransient, TextExportTransient)
 	TObjectPtr<UNiagaraSystem> NiagaraSystem;
 
+	/** Template niagara system asset used for this layout */
+	UPROPERTY(Transient)
+	mutable TObjectPtr<UNiagaraSystem> TemplateNiagaraSystem;
+
+	/** Base niagara system asset used for this layout */
+	UPROPERTY(Transient)
+	mutable TObjectPtr<UNiagaraSystem> BaseNiagaraSystem;
+
 	/** Mesh renderer in this niagara system */
 	UPROPERTY(Transient, DuplicateTransient, TextExportTransient)
 	TObjectPtr<UNiagaraMeshRendererProperties> MeshRenderer;
