Enable NDK 29 for Android, fix compilation issues
#Synced-CL 48505450

diff --git a/Engine/Config/Android/Android_SDK.json b/Engine/Config/Android/Android_SDK.json
--- a/Engine/Config/Android/Android_SDK.json
+++ b/Engine/Config/Android/Android_SDK.json
@@ -1,7 +1,7 @@
 {
 	"MainVersion": "r27c",
 	"MinVersion": "r27c",
-	"MaxVersion": "r28c",
+	"MaxVersion": "r29",
 
 	"AutoSDKDirectory": "-28",
 
diff --git a/Engine/Plugins/Animation/ControlRig/Source/ControlRigEditor/Private/AnimDetails/AnimDetailsProxyManager.cpp b/Engine/Plugins/Animation/ControlRig/Source/ControlRigEditor/Private/AnimDetails/AnimDetailsProxyManager.cpp
--- a/Engine/Plugins/Animation/ControlRig/Source/ControlRigEditor/Private/AnimDetails/AnimDetailsProxyManager.cpp
+++ b/Engine/Plugins/Animation/ControlRig/Source/ControlRigEditor/Private/AnimDetails/AnimDetailsProxyManager.cpp
@@ -276,7 +276,7 @@ void UAnimDetailsProxyManager::ForceUpdateProxies()
 		return;
 	}
 
-	FTSTicker::GetCoreTicker().RemoveTicker(RequestUpdateProxiesTickerHandle);
+	FTSTicker::RemoveTicker(RequestUpdateProxiesTickerHandle);
 	RequestUpdateProxiesTickerHandle.Reset();
 
 	const FControlRigEditMode* EditMode = static_cast<FControlRigEditMode*>(GLevelEditorModeTools().GetActiveMode(FControlRigEditMode::ModeName));
diff --git a/Engine/Plugins/Experimental/AIAssistant/Source/AIAssistant/Private/AIAssistantWebConnectionWidget.cpp b/Engine/Plugins/Experimental/AIAssistant/Source/AIAssistant/Private/AIAssistantWebConnectionWidget.cpp
--- a/Engine/Plugins/Experimental/AIAssistant/Source/AIAssistant/Private/AIAssistantWebConnectionWidget.cpp
+++ b/Engine/Plugins/Experimental/AIAssistant/Source/AIAssistant/Private/AIAssistantWebConnectionWidget.cpp
@@ -186,7 +186,7 @@ void SAIAssistantWebConnectionWidget::StopReconnecting()
 	
 	if (ReconnectingTicker.IsValid())
 	{
-		FTSTicker::GetCoreTicker().RemoveTicker(ReconnectingTicker);
+		FTSTicker::RemoveTicker(ReconnectingTicker);
 		ReconnectingTicker.Reset();
 	}
 
diff --git a/Engine/Plugins/Experimental/Fracture/Source/FractureEngine/Private/FractureEngineFracturing.cpp b/Engine/Plugins/Experimental/Fracture/Source/FractureEngine/Private/FractureEngineFracturing.cpp
--- a/Engine/Plugins/Experimental/Fracture/Source/FractureEngine/Private/FractureEngineFracturing.cpp
+++ b/Engine/Plugins/Experimental/Fracture/Source/FractureEngine/Private/FractureEngineFracturing.cpp
@@ -715,7 +715,7 @@ int32 FFractureEngineFracturing::SliceCutter(FManagedArrayCollection& InOutColle
 		// Proximity is invalidated.
 		ClearProximity(GeomCollection.Get());
 
-		int ResultGeometryIndex = CutMultipleWithMultiplePlanes(CuttingPlanes, InternalSurfaceMaterials, *GeomCollection, TransformSelectionArr, GroutVal, CollisionSampleSpacingVal, InRandomSeed, FTransform().Identity, true, nullptr, InSplitIslands);
+		int ResultGeometryIndex = CutMultipleWithMultiplePlanes(CuttingPlanes, InternalSurfaceMaterials, *GeomCollection, TransformSelectionArr, GroutVal, CollisionSampleSpacingVal, InRandomSeed, FTransform::Identity, true, nullptr, InSplitIslands);
 
 		UE::Private::FractureHelpers::ProcessNewlyFracturedBones(*GeomCollection, ResultGeometryIndex);
 
diff --git a/Engine/Plugins/Experimental/MessageBusTester/Source/MessageBusTester/Private/MessageBusTester.cpp b/Engine/Plugins/Experimental/MessageBusTester/Source/MessageBusTester/Private/MessageBusTester.cpp
--- a/Engine/Plugins/Experimental/MessageBusTester/Source/MessageBusTester/Private/MessageBusTester.cpp
+++ b/Engine/Plugins/Experimental/MessageBusTester/Source/MessageBusTester/Private/MessageBusTester.cpp
@@ -105,7 +105,7 @@ FMessageBusTester::~FMessageBusTester()
 
 	if (TickerHandle.IsValid())
 	{
-		FTSTicker::GetCoreTicker().RemoveTicker(TickerHandle);
+		FTSTicker::RemoveTicker(TickerHandle);
 		TickerHandle.Reset();
 	}
 	
diff --git a/Engine/Plugins/Messaging/UdpMessaging/Source/UdpMessaging/Private/Transport/UdpMessageProcessorLogger.cpp b/Engine/Plugins/Messaging/UdpMessaging/Source/UdpMessaging/Private/Transport/UdpMessageProcessorLogger.cpp
--- a/Engine/Plugins/Messaging/UdpMessaging/Source/UdpMessaging/Private/Transport/UdpMessageProcessorLogger.cpp
+++ b/Engine/Plugins/Messaging/UdpMessaging/Source/UdpMessaging/Private/Transport/UdpMessageProcessorLogger.cpp
@@ -140,7 +140,7 @@ struct FUdpMessageProcessorLoggerTyped
 	~FUdpMessageProcessorLoggerTyped()
 	{
 		UE::UdpMessaging::Private::Unregister(this);
-		FTSTicker::GetCoreTicker().RemoveTicker(TickHandler);
+		FTSTicker::RemoveTicker(TickHandler);
 		GameTick(0);
 
 		Flush();
diff --git a/Engine/Plugins/VirtualProduction/ClonerEffector/Source/ClonerEffector/Private/Cloner/CEClonerComponent.cpp b/Engine/Plugins/VirtualProduction/ClonerEffector/Source/ClonerEffector/Private/Cloner/CEClonerComponent.cpp
--- a/Engine/Plugins/VirtualProduction/ClonerEffector/Source/ClonerEffector/Private/Cloner/CEClonerComponent.cpp
+++ b/Engine/Plugins/VirtualProduction/ClonerEffector/Source/ClonerEffector/Private/Cloner/CEClonerComponent.cpp
@@ -195,10 +195,10 @@ void UCEClonerComponent::OnComponentDestroyed(bool bInDestroyingHierarchy)
 
 	FWorldDelegates::LevelAddedToWorld.RemoveAll(this);
 
-	FTSTicker::GetCoreTicker().RemoveTicker(ClonerTickerHandle);
+	FTSTicker::RemoveTicker(ClonerTickerHandle);
 	ClonerTickerHandle.Reset();
 
-	FTSTicker::GetCoreTicker().RemoveTicker(ConversionTickerHandle);
+	FTSTicker::RemoveTicker(ConversionTickerHandle);
 	ConversionTickerHandle.Reset();
 }
 
@@ -253,7 +253,7 @@ void UCEClonerComponent::UpdateDirtyMeshesAsync()
 
 	if (ConversionTickerHandle.IsValid())
 	{
-		FTSTicker::GetCoreTicker().RemoveTicker(ConversionTickerHandle);
+		FTSTicker::RemoveTicker(ConversionTickerHandle);
 	}
 
 	UE_LOG(LogCECloner, Verbose, TEXT("%s : Updating %i dirty actor meshes..."), *GetOwner()->GetActorNameOrLabel(), ClonerTree.DirtyItemAttachments.Num());
diff --git a/Engine/Plugins/VirtualProduction/OperatorStack/Source/OperatorStack/Private/Widgets/Tabs/OperatorStackEditorTabInstance.cpp b/Engine/Plugins/VirtualProduction/OperatorStack/Source/OperatorStack/Private/Widgets/Tabs/OperatorStackEditorTabInstance.cpp
--- a/Engine/Plugins/VirtualProduction/OperatorStack/Source/OperatorStack/Private/Widgets/Tabs/OperatorStackEditorTabInstance.cpp
+++ b/Engine/Plugins/VirtualProduction/OperatorStack/Source/OperatorStack/Private/Widgets/Tabs/OperatorStackEditorTabInstance.cpp
@@ -261,7 +261,7 @@ void FOperatorStackEditorTabInstance::OnSelectionSetChanged(const UTypedElementS
 	// Cancel previous request
 	if (SelectionDelegateHandle.IsValid())
 	{
-		FTSTicker::GetCoreTicker().RemoveTicker(SelectionDelegateHandle);
+		FTSTicker::RemoveTicker(SelectionDelegateHandle);
 		SelectionDelegateHandle.Reset();
 	}
 
diff --git a/Engine/Plugins/VirtualProduction/Text3D/Source/Text3D/Private/Text3DComponent.cpp b/Engine/Plugins/VirtualProduction/Text3D/Source/Text3D/Private/Text3DComponent.cpp
--- a/Engine/Plugins/VirtualProduction/Text3D/Source/Text3D/Private/Text3DComponent.cpp
+++ b/Engine/Plugins/VirtualProduction/Text3D/Source/Text3D/Private/Text3DComponent.cpp
@@ -80,7 +80,7 @@ void UText3DComponent::OnComponentDestroyed(bool bInDestroyingHierarchy)
 
 	TransformUpdated.RemoveAll(this);
 
-	FTSTicker::GetCoreTicker().RemoveTicker(TextUpdateHandle);
+	FTSTicker::RemoveTicker(TextUpdateHandle);
 	TextUpdateHandle.Reset();
 }
 
