Motion Design: fix issue where Effectors block saving a level

#Synced-CL 51764840

diff --git a/Engine/Plugins/VirtualProduction/ClonerEffector/Source/ClonerEffector/Private/Effector/CEEffectorActor.cpp b/Engine/Plugins/VirtualProduction/ClonerEffector/Source/ClonerEffector/Private/Effector/CEEffectorActor.cpp
--- a/Engine/Plugins/VirtualProduction/ClonerEffector/Source/ClonerEffector/Private/Effector/CEEffectorActor.cpp
+++ b/Engine/Plugins/VirtualProduction/ClonerEffector/Source/ClonerEffector/Private/Effector/CEEffectorActor.cpp
@@ -47,18 +47,22 @@ ACEEffectorActor::ACEEffectorActor()
 	EffectorComponent = CreateDefaultSubobject<UCEEffectorComponent>(TEXT("AvaEffectorComponent"));
 	SetRootComponent(EffectorComponent);
 
-	UDynamicMeshComponent* InnerVisualizerComponent = CreateDefaultSubobject<UDynamicMeshComponent>(TEXT("InnerVisualizerComponent"));
+#if WITH_EDITOR
+	UDynamicMeshComponent* InnerVisualizerComponent = CreateEditorOnlyDefaultSubobject<UDynamicMeshComponent>(TEXT("InnerVisualizerComponent"));
 	InnerVisualizerComponent->SetupAttachment(EffectorComponent);
 	InnerVisualizerComponent->bIsEditorOnly = true;
-#if WITH_EDITOR
-	EffectorComponent->AddVisualizerComponent(InnerVisualizerComponent);
-#endif
+	if (!HasAnyFlags(RF_ClassDefaultObject))
+	{
+		EffectorComponent->AddVisualizerComponent(InnerVisualizerComponent);	
+	}
 
-	UDynamicMeshComponent* OuterVisualizerComponent = CreateDefaultSubobject<UDynamicMeshComponent>(TEXT("OuterVisualizerComponent"));
+	UDynamicMeshComponent* OuterVisualizerComponent = CreateEditorOnlyDefaultSubobject<UDynamicMeshComponent>(TEXT("OuterVisualizerComponent"));
 	OuterVisualizerComponent->SetupAttachment(EffectorComponent);
 	OuterVisualizerComponent->bIsEditorOnly = true;
-#if WITH_EDITOR
-	EffectorComponent->AddVisualizerComponent(OuterVisualizerComponent);
+	if (!HasAnyFlags(RF_ClassDefaultObject))
+	{
+		EffectorComponent->AddVisualizerComponent(OuterVisualizerComponent);
+	}
 #endif
 }
 
