Motion Design: fix regression crash from cloner visualizers now being editor-only

#Synced-CL 51764840

diff --git a/Engine/Plugins/VirtualProduction/ClonerEffector/Source/ClonerEffector/Private/Effector/CEEffectorActor.cpp b/Engine/Plugins/VirtualProduction/ClonerEffector/Source/ClonerEffector/Private/Effector/CEEffectorActor.cpp
--- a/Engine/Plugins/VirtualProduction/ClonerEffector/Source/ClonerEffector/Private/Effector/CEEffectorActor.cpp
+++ b/Engine/Plugins/VirtualProduction/ClonerEffector/Source/ClonerEffector/Private/Effector/CEEffectorActor.cpp
@@ -48,19 +48,19 @@ ACEEffectorActor::ACEEffectorActor()
 	SetRootComponent(EffectorComponent);
 
 #if WITH_EDITOR
-	UDynamicMeshComponent* InnerVisualizerComponent = CreateEditorOnlyDefaultSubobject<UDynamicMeshComponent>(TEXT("InnerVisualizerComponent"));
-	InnerVisualizerComponent->SetupAttachment(EffectorComponent);
-	InnerVisualizerComponent->bIsEditorOnly = true;
-	if (!HasAnyFlags(RF_ClassDefaultObject))
+	if (GIsEditor && !HasAnyFlags(RF_ClassDefaultObject))
 	{
+		UDynamicMeshComponent* InnerVisualizerComponent = CreateEditorOnlyDefaultSubobject<UDynamicMeshComponent>(TEXT("InnerVisualizerComponent"));
+		InnerVisualizerComponent->SetupAttachment(EffectorComponent);
+		InnerVisualizerComponent->bIsEditorOnly = true;
 		EffectorComponent->AddVisualizerComponent(InnerVisualizerComponent);	
 	}
 
-	UDynamicMeshComponent* OuterVisualizerComponent = CreateEditorOnlyDefaultSubobject<UDynamicMeshComponent>(TEXT("OuterVisualizerComponent"));
-	OuterVisualizerComponent->SetupAttachment(EffectorComponent);
-	OuterVisualizerComponent->bIsEditorOnly = true;
-	if (!HasAnyFlags(RF_ClassDefaultObject))
+	if (GIsEditor && !HasAnyFlags(RF_ClassDefaultObject))
 	{
+		UDynamicMeshComponent* OuterVisualizerComponent = CreateEditorOnlyDefaultSubobject<UDynamicMeshComponent>(TEXT("OuterVisualizerComponent"));
+		OuterVisualizerComponent->SetupAttachment(EffectorComponent);
+		OuterVisualizerComponent->bIsEditorOnly = true;
 		EffectorComponent->AddVisualizerComponent(OuterVisualizerComponent);
 	}
 #endif
