Motion Design: changed cloners to use mesh array data interface

#virtualized
#Synced-CL 51764840

diff --git a/Engine/Plugins/VirtualProduction/ClonerEffector/Source/ClonerEffector/Private/Cloner/CEClonerComponent.cpp b/Engine/Plugins/VirtualProduction/ClonerEffector/Source/ClonerEffector/Private/Cloner/CEClonerComponent.cpp
--- a/Engine/Plugins/VirtualProduction/ClonerEffector/Source/ClonerEffector/Private/Cloner/CEClonerComponent.cpp
+++ b/Engine/Plugins/VirtualProduction/ClonerEffector/Source/ClonerEffector/Private/Cloner/CEClonerComponent.cpp
@@ -9,6 +9,7 @@
 #include "Cloner/Layouts/CEClonerLayoutBase.h"
 #include "Components/BillboardComponent.h"
 #include "Components/SkeletalMeshComponent.h"
+#include "DataInterface/NiagaraDataInterfaceArrayMesh.h"
 #include "Engine/Level.h"
 #include "Engine/StaticMeshActor.h"
 #include "Engine/Texture2D.h"
@@ -17,6 +18,7 @@
 #include "Materials/Material.h"
 #include "Materials/MaterialInstanceDynamic.h"
 #include "Misc/EnumerateRange.h"
+#include "NiagaraFunctionLibrary.h"
 #include "NiagaraMeshRendererProperties.h"
 #include "NiagaraSystem.h"
 #include "Settings/CEClonerEffectorSettings.h"
@@ -765,14 +767,10 @@ void UCEClonerComponent::SetVisualizerSpriteVisible(bool bInVisible)
 
 int32 UCEClonerComponent::GetMeshCount() const
 {
-	if (const UCEClonerLayoutBase* LayoutSystem = GetActiveLayout())
+	if (const UNiagaraDataInterfaceArrayMesh* const MeshArrayDI = UNiagaraFunctionLibrary::GetDataInterface<UNiagaraDataInterfaceArrayMesh>(const_cast<UCEClonerComponent*>(this), UE::ClonerEffector::MeshArrayName))
 	{
-		if (const UNiagaraMeshRendererProperties* MeshRenderer = LayoutSystem->GetMeshRenderer())
-		{
-			return MeshRenderer->Meshes.Num();
-		}
+		return MeshArrayDI->MeshData.Num();
 	}
-
 	return 0;
 }
 
@@ -1420,6 +1418,12 @@ void UCEClonerComponent::UpdateClonerMeshes()
 		return;
 	}
 
+	const UNiagaraDataInterfaceArrayMesh* const MeshArrayDI = UNiagaraFunctionLibrary::GetDataInterface<UNiagaraDataInterfaceArrayMesh>(this, UE::ClonerEffector::MeshArrayName);
+	if (!MeshArrayDI)
+	{
+		return;
+	}
+
 	UNiagaraSystem* ActiveSystem = GetAsset();
 
 	if (!ActiveSystem || !ActiveLayout)
@@ -1441,14 +1445,15 @@ void UCEClonerComponent::UpdateClonerMeshes()
 		return;
 	}
 
-	bool bMeshChanged = MeshRenderer->Meshes.Num() != ClonerTree.RootItemAttachments.Num();
+	TArray<FNiagaraMeshRendererMeshPropertiesBase> MeshData = MeshArrayDI->MeshData;
+	bool bMeshChanged = MeshData.Num() != ClonerTree.RootItemAttachments.Num();
 
 	if (ClonerTree.bItemAttachmentsDirty)
 	{
 		// Resize mesh array properly
-		if (MeshRenderer->Meshes.Num() > ClonerTree.RootItemAttachments.Num())
+		if (MeshData.Num() > ClonerTree.RootItemAttachments.Num())
 		{
-			MeshRenderer->Meshes.SetNum(ClonerTree.RootItemAttachments.Num());
+			MeshData.SetNum(ClonerTree.RootItemAttachments.Num());
 		}
 
 		// Set baked meshes in mesh renderer array
@@ -1457,7 +1462,7 @@ void UCEClonerComponent::UpdateClonerMeshes()
 			UStaticMesh* StaticMesh = RootItem->MergedBakedMesh.Get();
 			const int32 Index = RootItem.GetIndex();
 
-			FNiagaraMeshRendererMeshProperties& MeshProperties = !MeshRenderer->Meshes.IsValidIndex(Index) ? MeshRenderer->Meshes.AddDefaulted_GetRef() : MeshRenderer->Meshes[Index];
+			FNiagaraMeshRendererMeshPropertiesBase& MeshProperties = !MeshData.IsValidIndex(Index) ? MeshData.AddDefaulted_GetRef() : MeshData[Index];
 			bMeshChanged |= MeshProperties.Mesh != StaticMesh;
 
 			if (StaticMesh && StaticMesh->GetNumTriangles(0) > 0)
@@ -1488,25 +1493,17 @@ void UCEClonerComponent::UpdateClonerMeshes()
 	}
 
 	// Extensions could override mesh renderer meshes array
-	bMeshChanged |= ClonerTree.RootItemAttachments.Num() != MeshRenderer->Meshes.Num();
+	bMeshChanged |= ClonerTree.RootItemAttachments.Num() != MeshData.Num();
 
 	if (bMeshChanged)
 	{
-		UE_LOG(LogCECloner, Verbose, TEXT("%s : Cloner mesh updated - %i cached meshes - %i rendered meshes"), *ClonerActor->GetActorNameOrLabel(), ClonerTree.RootItemAttachments.Num(), MeshRenderer->Meshes.Num())
+		UE_LOG(LogCECloner, Verbose, TEXT("%s : Cloner mesh updated - %i cached meshes - %i rendered meshes"), *ClonerActor->GetActorNameOrLabel(), ClonerTree.RootItemAttachments.Num(), MeshData.Num())
 	}
 
 	// Set new number of meshes in renderer
-	SetIntParameter(TEXT("MeshNum"), MeshRenderer->Meshes.Num());
-
-#if WITH_EDITORONLY_DATA
-	MeshRenderer->OnMeshChanged();
-
-	// Used by other data interfaces to update their cached data
-	MeshRenderer->OnChanged().Broadcast();
-#else
-	FNiagaraSystemUpdateContext ReregisterContext(ActiveSystem, /** bReset */ true);
-#endif
+	SetIntParameter(TEXT("MeshNum"), MeshData.Num());
 
+	UNiagaraDataInterfaceArrayMesh::SetNiagaraArrayMesh(this, UE::ClonerEffector::MeshArrayName, MeshData);
 	OnClonerMeshUpdatedDelegate.Broadcast(this);
 }
 
diff --git a/Engine/Plugins/VirtualProduction/ClonerEffector/Source/ClonerEffector/Private/Cloner/Extensions/CEClonerCollisionExtension.cpp b/Engine/Plugins/VirtualProduction/ClonerEffector/Source/ClonerEffector/Private/Cloner/Extensions/CEClonerCollisionExtension.cpp
--- a/Engine/Plugins/VirtualProduction/ClonerEffector/Source/ClonerEffector/Private/Cloner/Extensions/CEClonerCollisionExtension.cpp
+++ b/Engine/Plugins/VirtualProduction/ClonerEffector/Source/ClonerEffector/Private/Cloner/Extensions/CEClonerCollisionExtension.cpp
@@ -4,9 +4,11 @@
 
 #include "Cloner/CEClonerComponent.h"
 #include "Components/StaticMeshComponent.h"
+#include "DataInterface/NiagaraDataInterfaceArrayMesh.h"
 #include "Engine/World.h"
 #include "Materials/MaterialInstanceDynamic.h"
 #include "NiagaraDataInterfaceArrayFloat.h"
+#include "NiagaraFunctionLibrary.h"
 #include "NiagaraMeshRendererProperties.h"
 #include "NiagaraSystem.h"
 
@@ -139,15 +141,13 @@ void UCEClonerCollisionExtension::OnExtensionParametersChanged(UCEClonerComponen
 	// Adjust size based on attached mesh count
 	CollisionRadii.SetNum(InComponent->GetMeshCount());
 
-	if (const UCEClonerLayoutBase* LayoutSystem = GetClonerLayout())
+	if (const UNiagaraDataInterfaceArrayMesh* const MeshArrayDI = UNiagaraFunctionLibrary::GetDataInterface<UNiagaraDataInterfaceArrayMesh>(InComponent, UE::ClonerEffector::MeshArrayName))
 	{
 		if (CollisionRadiusMode != ECEClonerCollisionRadiusMode::Manual)
 		{
-			UNiagaraMeshRendererProperties* MeshRenderer = LayoutSystem->GetMeshRenderer();
-
 			for (int32 Idx = 0; Idx < CollisionRadii.Num(); Idx++)
 			{
-				const FNiagaraMeshRendererMeshProperties& MeshProperties = MeshRenderer->Meshes[Idx];
+				const FNiagaraMeshRendererMeshPropertiesBase& MeshProperties = MeshArrayDI->MeshData[Idx];
 				FBoxSphereBounds MeshBounds(ForceInitToZero);
 				const FTransform BoundTransform(MeshProperties.Rotation, MeshProperties.PivotOffset, MeshProperties.Scale);
 
diff --git a/Engine/Plugins/VirtualProduction/ClonerEffector/Source/ClonerEffector/Private/Cloner/Extensions/CEClonerMeshRendererExtension.cpp b/Engine/Plugins/VirtualProduction/ClonerEffector/Source/ClonerEffector/Private/Cloner/Extensions/CEClonerMeshRendererExtension.cpp
--- a/Engine/Plugins/VirtualProduction/ClonerEffector/Source/ClonerEffector/Private/Cloner/Extensions/CEClonerMeshRendererExtension.cpp
+++ b/Engine/Plugins/VirtualProduction/ClonerEffector/Source/ClonerEffector/Private/Cloner/Extensions/CEClonerMeshRendererExtension.cpp
@@ -4,8 +4,10 @@
 
 #include "Cloner/CEClonerComponent.h"
 #include "Cloner/Logs/CEClonerLogs.h"
+#include "DataInterface/NiagaraDataInterfaceArrayMesh.h"
 #include "Materials/Material.h"
 #include "Materials/MaterialInterface.h"
+#include "NiagaraFunctionLibrary.h"
 #include "NiagaraSystem.h"
 #include "Settings/CEClonerEffectorSettings.h"
 #include "UObject/ConstructorHelpers.h"
@@ -143,21 +145,19 @@ int32 UCEClonerMeshRendererExtension::GetClonerMeshesMaterialCount() const
 {
 	int32 MaterialCount = 0;
 
-	UCEClonerLayoutBase* LayoutSystem = GetClonerLayout();
-
-	if (!IsValid(LayoutSystem))
+	UCEClonerComponent* const ClonerComponent = GetClonerComponent();
+	if (!ClonerComponent)
 	{
 		return MaterialCount;
 	}
 
-	UNiagaraMeshRendererProperties* MeshRenderer = LayoutSystem->GetMeshRenderer();
-
-	if (!IsValid(MeshRenderer))
+	const UNiagaraDataInterfaceArrayMesh* const MeshArrayDI = UNiagaraFunctionLibrary::GetDataInterface<UNiagaraDataInterfaceArrayMesh>(ClonerComponent, UE::ClonerEffector::MeshArrayName);
+	if (!MeshArrayDI)
 	{
 		return MaterialCount;
 	}
 
-	for (const FNiagaraMeshRendererMeshProperties& MeshProperties : MeshRenderer->Meshes)
+	for (const FNiagaraMeshRendererMeshPropertiesBase& MeshProperties : MeshArrayDI->MeshData)
 	{
 		if (!MeshProperties.Mesh)
 		{
@@ -222,27 +222,26 @@ void UCEClonerMeshRendererExtension::OnClonerMeshesUpdated()
 {
 	Super::OnClonerMeshesUpdated();
 
-	UCEClonerComponent* ClonerComponent = GetClonerComponent();
-
-	if (!IsValid(ClonerComponent))
+	UCEClonerComponent* const ClonerComponent = GetClonerComponent();
+	if (!ClonerComponent)
 	{
 		return;
 	}
 
-	const UCEClonerLayoutBase* Layout = GetClonerLayout();
-
-	if (!IsValid(Layout))
+	UCEClonerLayoutBase* const Layout = GetClonerLayout();
+	if (!Layout)
 	{
 		return;
 	}
 
-	UNiagaraMeshRendererProperties* MeshRenderer = Layout->GetMeshRenderer();
-
-	if (!IsValid(MeshRenderer))
+	UNiagaraMeshRendererProperties* const MeshRenderer = Layout->GetMeshRenderer();
+	if (!MeshRenderer)
 	{
 		return;
 	}
 
+	// Clear the meshes as the mesh array DI is going to be used. This should be always empty for new systems, but old systems could have lingering meshes here.
+	MeshRenderer->Meshes.Empty();
 	MeshRenderer->FacingMode = MeshFacingMode;
 	MeshRenderer->bCastShadows = bMeshCastShadows;
 	MeshRenderer->SortMode = bSortTranslucentParticles ? ENiagaraSortMode::ViewDepth : ENiagaraSortMode::None;
@@ -251,16 +250,20 @@ void UCEClonerMeshRendererExtension::OnClonerMeshesUpdated()
 	if (ClonerComponent->GetAttachmentCount() == 0)
 	{
 		const TArray<TObjectPtr<UStaticMesh>>& NewDefaultMeshes = GetDefaultMeshes();
-		MeshRenderer->Meshes.SetNum(NewDefaultMeshes.Num());
+
+		TArray<FNiagaraMeshRendererMeshPropertiesBase> Meshes;
+		Meshes.SetNum(NewDefaultMeshes.Num());
 
 		for (int32 MeshIndex = 0; MeshIndex < NewDefaultMeshes.Num(); MeshIndex++)
 		{
 			UStaticMesh* DefaultMesh = NewDefaultMeshes[MeshIndex].Get();
-			FNiagaraMeshRendererMeshProperties& MeshProperties = MeshRenderer->Meshes[MeshIndex];
+			FNiagaraMeshRendererMeshPropertiesBase& MeshProperties = Meshes[MeshIndex];
 			MeshProperties.Mesh = DefaultMesh && DefaultMesh->GetNumTriangles(0) > 0 ? DefaultMesh : nullptr;
 			MeshProperties.Scale = ClonerComponent->GetGlobalScale();
 			MeshProperties.Rotation = ClonerComponent->GetGlobalRotation();
 		}
+
+		UNiagaraDataInterfaceArrayMesh::SetNiagaraArrayMesh(ClonerComponent, UE::ClonerEffector::MeshArrayName, Meshes);
 	}
 
 	// Set material override
diff --git a/Engine/Plugins/VirtualProduction/ClonerEffector/Source/ClonerEffector/Private/Cloner/Layouts/CEClonerLayoutBase.cpp b/Engine/Plugins/VirtualProduction/ClonerEffector/Source/ClonerEffector/Private/Cloner/Layouts/CEClonerLayoutBase.cpp
--- a/Engine/Plugins/VirtualProduction/ClonerEffector/Source/ClonerEffector/Private/Cloner/Layouts/CEClonerLayoutBase.cpp
+++ b/Engine/Plugins/VirtualProduction/ClonerEffector/Source/ClonerEffector/Private/Cloner/Layouts/CEClonerLayoutBase.cpp
@@ -7,6 +7,7 @@
 #include "Cloner/CEClonerComponent.h"
 #include "Cloner/Extensions/CEClonerExtensionBase.h"
 #include "Cloner/Logs/CEClonerLogs.h"
+#include "DataInterface/NiagaraDataInterfaceArrayMesh.h"
 #include "Engine/Level.h"
 #include "Engine/World.h"
 #include "Misc/Guid.h"
@@ -193,10 +194,11 @@ bool UCEClonerLayoutBase::UnloadLayout()
 		return false;
 	}
 
+	UNiagaraDataInterfaceArrayMesh::SetNiagaraArrayMesh(GetClonerComponent(), UE::ClonerEffector::MeshArrayName, {});
+
+	// Clear the meshes as the mesh array DI is going to be used. This should be always empty for new systems, but old systems could have lingering meshes here.
 	MeshRenderer->Meshes.Empty();
 #if WITH_EDITOR
-	MeshRenderer->OnMeshChanged();
-	MeshRenderer->OnChanged().Broadcast();
 	NiagaraSystem->KillAllActiveCompilations();
 #endif
 	NiagaraSystem->RemoveFromRoot();
@@ -415,13 +417,8 @@ void UCEClonerLayoutBase::CacheMeshRenderer()
 			{
 				if (UNiagaraMeshRendererProperties* EmitterMeshRenderer = Cast<UNiagaraMeshRendererProperties>(EmitterRenderer))
 				{
-					EmitterMeshRenderer->Meshes.Empty();
-#if WITH_EDITOR
-					EmitterMeshRenderer->OnMeshChanged();
-#endif
-
+					UNiagaraDataInterfaceArrayMesh::SetNiagaraArrayMesh(GetClonerComponent(), UE::ClonerEffector::MeshArrayName, {});
 					MeshRenderer = EmitterMeshRenderer;
-
 					return;
 				}
 			}
diff --git a/Engine/Plugins/VirtualProduction/ClonerEffector/Source/ClonerEffector/Public/CEClonerEffectorShared.h b/Engine/Plugins/VirtualProduction/ClonerEffector/Source/ClonerEffector/Public/CEClonerEffectorShared.h
--- a/Engine/Plugins/VirtualProduction/ClonerEffector/Source/ClonerEffector/Public/CEClonerEffectorShared.h
+++ b/Engine/Plugins/VirtualProduction/ClonerEffector/Source/ClonerEffector/Public/CEClonerEffectorShared.h
@@ -463,6 +463,9 @@ struct FCEExtensionSection
 
 namespace UE::ClonerEffector
 {
+	/** Name for the mesh array data interface to use */
+	static const FLazyName MeshArrayName = TEXT("MeshArray"); 
+
 #if WITH_EDITOR
 	namespace EditorSection
 	{
