Motion Design: removed SetAsset(null) for cloners to re-use the cached asset if unchanged. Added `ClonerEffector.UpdateSystemAfterLoad` cvar for UpdateSystemAfterLoad() to be called on cloner niagara systems for this operation to be done in loading time rather than in activation time, as activation time is likely to be more time-sensitive than loading.
Additionally, skip doing any visualizer work in -game.

#Synced-CL 51764840

diff --git a/Engine/Plugins/VirtualProduction/ClonerEffector/Source/ClonerEffector/Private/Cloner/CEClonerComponent.cpp b/Engine/Plugins/VirtualProduction/ClonerEffector/Source/ClonerEffector/Private/Cloner/CEClonerComponent.cpp
--- a/Engine/Plugins/VirtualProduction/ClonerEffector/Source/ClonerEffector/Private/Cloner/CEClonerComponent.cpp
+++ b/Engine/Plugins/VirtualProduction/ClonerEffector/Source/ClonerEffector/Private/Cloner/CEClonerComponent.cpp
@@ -497,10 +497,12 @@ void UCEClonerComponent::InitializeCloner()
 
 	bClonerInitialized = true;
 
-	SetAsset(nullptr);
-
 #if WITH_EDITOR
-	OnVisualizerSpriteVisibleChanged();
+	extern CORE_API bool GIsEditor;
+	if (GIsEditor)
+	{
+		OnVisualizerSpriteVisibleChanged();
+	}
 
 	// Skip init for preview actor
 	if (Owner->bIsEditorPreviewActor)
diff --git a/Engine/Plugins/VirtualProduction/ClonerEffector/Source/ClonerEffector/Private/Cloner/Layouts/CEClonerLayoutBase.cpp b/Engine/Plugins/VirtualProduction/ClonerEffector/Source/ClonerEffector/Private/Cloner/Layouts/CEClonerLayoutBase.cpp
--- a/Engine/Plugins/VirtualProduction/ClonerEffector/Source/ClonerEffector/Private/Cloner/Layouts/CEClonerLayoutBase.cpp
+++ b/Engine/Plugins/VirtualProduction/ClonerEffector/Source/ClonerEffector/Private/Cloner/Layouts/CEClonerLayoutBase.cpp
@@ -19,6 +19,18 @@
 #include "UObject/Package.h"
 #include "UObject/UObjectGlobals.h"
 
+namespace UE::ClonerEffector::Private
+{
+#if WITH_EDITOR
+	TAutoConsoleVariable<bool> CVarUpdateSystemAfterLoad(
+		TEXT("ClonerEffector.UpdateSystemAfterLoad"),
+		false,
+		TEXT("Calls UpdateSystemAfterLoad() on a Cloner Layout's PostLoad, so that this is done in loading time instead of activation time.\n")
+		TEXT("Only relevant for -game instances (in editor builds).\n")
+		TEXT("Non-editor builds do not require this as they already call this in the system PostLoad."));
+#endif
+}
+
 bool UCEClonerLayoutBase::IsLayoutValid() const
 {
 	if (LayoutName.IsNone() || LayoutAssetPath.IsEmpty())
@@ -341,6 +353,16 @@ void UCEClonerLayoutBase::PostLoad()
 		// so look for it in outer chain otherwise it will trigger a world GC leak when switching level
 		CleanOwnedSystem();
 	}
+
+#if WITH_EDITOR
+	extern CORE_API bool GIsEditor;
+	if (NiagaraSystem && !GIsEditor && UE::ClonerEffector::Private::CVarUpdateSystemAfterLoad.GetValueOnAnyThread())
+	{
+		// This is done in loading time to not pay this cost in activation time
+		// This is already called in UNiagaraSystem::PostLoad in non-editor builds.
+		NiagaraSystem->UpdateSystemAfterLoad();
+	}
+#endif
 }
 
 #if WITH_EDITOR
diff --git a/Engine/Plugins/VirtualProduction/ClonerEffector/Source/ClonerEffector/Private/Effector/CEEffectorComponent.cpp b/Engine/Plugins/VirtualProduction/ClonerEffector/Source/ClonerEffector/Private/Effector/CEEffectorComponent.cpp
--- a/Engine/Plugins/VirtualProduction/ClonerEffector/Source/ClonerEffector/Private/Effector/CEEffectorComponent.cpp
+++ b/Engine/Plugins/VirtualProduction/ClonerEffector/Source/ClonerEffector/Private/Effector/CEEffectorComponent.cpp
@@ -716,6 +716,12 @@ void UCEEffectorComponent::OnEffectsChanged()
 #if WITH_EDITOR
 void UCEEffectorComponent::OnVisualizerOptionsChanged()
 {
+	// Only run this while in editor
+	if (!GIsEditor)
+	{
+		return;
+	}
+
 	for (const TWeakObjectPtr<UDynamicMeshComponent>& VisualizerComponentWeak : VisualizerComponentsWeak)
 	{
 		if (UDynamicMeshComponent* VisualizerComponent = VisualizerComponentWeak.Get())
