Motion Design: fix issues with cloner updating mesh with mesh array data interface which triggers a 'modify' causing the level to be dirtied.
This is fixed by not going through this code path if there is nothing to change.

#Synced-CL 52352995

diff --git a/Engine/Plugins/VirtualProduction/ClonerEffector/Source/ClonerEffector/Private/Cloner/CEClonerComponent.cpp b/Engine/Plugins/VirtualProduction/ClonerEffector/Source/ClonerEffector/Private/Cloner/CEClonerComponent.cpp
--- a/Engine/Plugins/VirtualProduction/ClonerEffector/Source/ClonerEffector/Private/Cloner/CEClonerComponent.cpp
+++ b/Engine/Plugins/VirtualProduction/ClonerEffector/Source/ClonerEffector/Private/Cloner/CEClonerComponent.cpp
@@ -498,7 +498,6 @@ void UCEClonerComponent::InitializeCloner()
 	bClonerInitialized = true;
 
 #if WITH_EDITOR
-	extern CORE_API bool GIsEditor;
 	if (GIsEditor)
 	{
 		OnVisualizerSpriteVisibleChanged();
@@ -1453,35 +1452,34 @@ void UCEClonerComponent::UpdateClonerMeshes()
 	if (ClonerTree.bItemAttachmentsDirty)
 	{
 		// Resize mesh array properly
-		if (MeshData.Num() > ClonerTree.RootItemAttachments.Num())
-		{
-			MeshData.SetNum(ClonerTree.RootItemAttachments.Num());
-		}
+		MeshData.SetNum(ClonerTree.RootItemAttachments.Num());
 
 		// Set baked meshes in mesh renderer array
 		for (TConstEnumerateRef<FCEClonerAttachmentRootItem> RootItem : EnumerateRange(ClonerTree.RootItemAttachments))
 		{
-			UStaticMesh* StaticMesh = RootItem->MergedBakedMesh.Get();
-			const int32 Index = RootItem.GetIndex();
+			FNiagaraMeshRendererMeshPropertiesBase& MeshProperties = MeshData[RootItem.GetIndex()];
 
-			FNiagaraMeshRendererMeshPropertiesBase& MeshProperties = !MeshData.IsValidIndex(Index) ? MeshData.AddDefaulted_GetRef() : MeshData[Index];
-			bMeshChanged |= MeshProperties.Mesh != StaticMesh;
-
-			if (StaticMesh && StaticMesh->GetNumTriangles(0) > 0)
+			UStaticMesh* TargetStaticMesh = RootItem->MergedBakedMesh.Get();
+			if (TargetStaticMesh && TargetStaticMesh->GetNumTriangles(0) <= 0)
 			{
-				MeshProperties.Mesh = StaticMesh;
-			}
-			else
-			{
-				MeshProperties.Mesh = nullptr;
+				TargetStaticMesh = nullptr;
 			}
+			bMeshChanged = bMeshChanged || MeshProperties.Mesh != TargetStaticMesh;
+			MeshProperties.Mesh = TargetStaticMesh;
 
 			if (const AActor* RootActor = RootItem->RootActor.Get())
 			{
 				if (const USceneComponent* SceneComponent = RootActor->GetRootComponent())
 				{
-					MeshProperties.Rotation = GlobalRotation + SceneComponent->GetRelativeRotation();
-					MeshProperties.Scale = GlobalScale * SceneComponent->GetRelativeScale3D();
+					const FRotator TargetRotation = GlobalRotation + SceneComponent->GetRelativeRotation();
+					const FVector TargetScale = GlobalScale * SceneComponent->GetRelativeScale3D();
+
+					bMeshChanged = bMeshChanged 
+						|| !MeshProperties.Rotation.Equals(TargetRotation)
+						|| !MeshProperties.Scale.Equals(TargetScale);
+
+					MeshProperties.Rotation = TargetRotation;
+					MeshProperties.Scale = TargetScale;
 				}
 			}
 		}
@@ -1491,17 +1489,23 @@ void UCEClonerComponent::UpdateClonerMeshes()
 
 	for (const TObjectPtr<UCEClonerExtensionBase>& ActiveExtension : ActiveExtensions)
 	{
-		ActiveExtension->OnClonerMeshesUpdated();
+		// Extensions can override mesh renderer meshes array
+		bMeshChanged |= ActiveExtension->OnUpdateClonerMeshes(MeshData);
 	}
 
-	// Extensions could override mesh renderer meshes array
-	bMeshChanged |= ClonerTree.RootItemAttachments.Num() != MeshData.Num();
+	bMeshChanged = bMeshChanged || ClonerTree.RootItemAttachments.Num() != MeshData.Num();
+	if (!bMeshChanged)
+	{
+		return;
+	}
 
-	if (bMeshChanged)
+	for (const TObjectPtr<UCEClonerExtensionBase>& ActiveExtension : ActiveExtensions)
 	{
-		UE_LOG(LogCECloner, Verbose, TEXT("%s : Cloner mesh updated - %i cached meshes - %i rendered meshes"), *ClonerActor->GetActorNameOrLabel(), ClonerTree.RootItemAttachments.Num(), MeshData.Num())
+		ActiveExtension->OnClonerMeshesUpdated();
 	}
 
+	UE_LOG(LogCECloner, Verbose, TEXT("%s : Cloner mesh updated - %i cached meshes - %i rendered meshes"), *ClonerActor->GetActorNameOrLabel(), ClonerTree.RootItemAttachments.Num(), MeshData.Num());
+
 	// Set new number of meshes in renderer
 	SetIntParameter(TEXT("MeshNum"), MeshData.Num());
 
diff --git a/Engine/Plugins/VirtualProduction/ClonerEffector/Source/ClonerEffector/Private/Cloner/Extensions/CEClonerMeshRendererExtension.cpp b/Engine/Plugins/VirtualProduction/ClonerEffector/Source/ClonerEffector/Private/Cloner/Extensions/CEClonerMeshRendererExtension.cpp
--- a/Engine/Plugins/VirtualProduction/ClonerEffector/Source/ClonerEffector/Private/Cloner/Extensions/CEClonerMeshRendererExtension.cpp
+++ b/Engine/Plugins/VirtualProduction/ClonerEffector/Source/ClonerEffector/Private/Cloner/Extensions/CEClonerMeshRendererExtension.cpp
@@ -218,26 +218,26 @@ void UCEClonerMeshRendererExtension::OnOverrideMaterialOptionsChanged()
 	}
 }
 
-void UCEClonerMeshRendererExtension::OnClonerMeshesUpdated()
+bool UCEClonerMeshRendererExtension::OnUpdateClonerMeshes(TArray<FNiagaraMeshRendererMeshPropertiesBase>& InOutMeshes)
 {
-	Super::OnClonerMeshesUpdated();
+	Super::OnUpdateClonerMeshes(InOutMeshes);
 
 	UCEClonerComponent* const ClonerComponent = GetClonerComponent();
 	if (!ClonerComponent)
 	{
-		return;
+		return false;
 	}
 
 	UCEClonerLayoutBase* const Layout = GetClonerLayout();
 	if (!Layout)
 	{
-		return;
+		return false;
 	}
 
 	UNiagaraMeshRendererProperties* const MeshRenderer = Layout->GetMeshRenderer();
 	if (!MeshRenderer)
 	{
-		return;
+		return false;
 	}
 
 	// Clear the meshes as the mesh array DI is going to be used. This should be always empty for new systems, but old systems could have lingering meshes here.
@@ -246,29 +246,42 @@ void UCEClonerMeshRendererExtension::OnClonerMeshesUpdated()
 	MeshRenderer->bCastShadows = bMeshCastShadows;
 	MeshRenderer->SortMode = bSortTranslucentParticles ? ENiagaraSortMode::ViewDepth : ENiagaraSortMode::None;
 
+	bool bMeshDataChanged = false;
+
 	// Use default meshes if nothing is attached
 	if (ClonerComponent->GetAttachmentCount() == 0)
 	{
 		const TArray<TObjectPtr<UStaticMesh>>& NewDefaultMeshes = GetDefaultMeshes();
 
-		TArray<FNiagaraMeshRendererMeshPropertiesBase> Meshes;
-		Meshes.SetNum(NewDefaultMeshes.Num());
+		bMeshDataChanged = InOutMeshes.Num() != NewDefaultMeshes.Num();
+		InOutMeshes.SetNum(NewDefaultMeshes.Num());
 
 		for (int32 MeshIndex = 0; MeshIndex < NewDefaultMeshes.Num(); MeshIndex++)
 		{
 			UStaticMesh* DefaultMesh = NewDefaultMeshes[MeshIndex].Get();
-			FNiagaraMeshRendererMeshPropertiesBase& MeshProperties = Meshes[MeshIndex];
-			MeshProperties.Mesh = DefaultMesh && DefaultMesh->GetNumTriangles(0) > 0 ? DefaultMesh : nullptr;
+			if (DefaultMesh && DefaultMesh->GetNumTriangles(0) <= 0)
+			{
+				DefaultMesh = nullptr;
+			}
+
+			FNiagaraMeshRendererMeshPropertiesBase& MeshProperties = InOutMeshes[MeshIndex];
+
+			bMeshDataChanged = bMeshDataChanged
+				|| MeshProperties.Mesh != DefaultMesh 
+				|| !MeshProperties.Scale.Equals(ClonerComponent->GetGlobalScale()) 
+				|| !MeshProperties.Rotation.Equals(ClonerComponent->GetGlobalRotation());
+
+			MeshProperties.Mesh = DefaultMesh;
 			MeshProperties.Scale = ClonerComponent->GetGlobalScale();
 			MeshProperties.Rotation = ClonerComponent->GetGlobalRotation();
 		}
-
-		UNiagaraDataInterfaceArrayMesh::SetNiagaraArrayMesh(ClonerComponent, UE::ClonerEffector::MeshArrayName, Meshes);
 	}
 
 	// Set material override
 	MeshRenderer->bOverrideMaterials = bUseOverrideMaterial || bVisualizeEffectors;
 	MeshRenderer->OverrideMaterials = GetOverrideMeshesMaterials();
+
+	return bMeshDataChanged;
 }
 
 #if WITH_EDITOR
diff --git a/Engine/Plugins/VirtualProduction/ClonerEffector/Source/ClonerEffector/Private/Cloner/Layouts/CEClonerLayoutBase.cpp b/Engine/Plugins/VirtualProduction/ClonerEffector/Source/ClonerEffector/Private/Cloner/Layouts/CEClonerLayoutBase.cpp
--- a/Engine/Plugins/VirtualProduction/ClonerEffector/Source/ClonerEffector/Private/Cloner/Layouts/CEClonerLayoutBase.cpp
+++ b/Engine/Plugins/VirtualProduction/ClonerEffector/Source/ClonerEffector/Private/Cloner/Layouts/CEClonerLayoutBase.cpp
@@ -355,7 +355,6 @@ void UCEClonerLayoutBase::PostLoad()
 	}
 
 #if WITH_EDITOR
-	extern CORE_API bool GIsEditor;
 	if (NiagaraSystem && !GIsEditor && UE::ClonerEffector::Private::CVarUpdateSystemAfterLoad.GetValueOnAnyThread())
 	{
 		// This is done in loading time to not pay this cost in activation time
diff --git a/Engine/Plugins/VirtualProduction/ClonerEffector/Source/ClonerEffector/Public/Cloner/Extensions/CEClonerExtensionBase.h b/Engine/Plugins/VirtualProduction/ClonerEffector/Source/ClonerEffector/Public/Cloner/Extensions/CEClonerExtensionBase.h
--- a/Engine/Plugins/VirtualProduction/ClonerEffector/Source/ClonerEffector/Public/Cloner/Extensions/CEClonerExtensionBase.h
+++ b/Engine/Plugins/VirtualProduction/ClonerEffector/Source/ClonerEffector/Public/Cloner/Extensions/CEClonerExtensionBase.h
@@ -8,6 +8,7 @@
 
 class UCEClonerComponent;
 class UCEClonerLayoutBase;
+struct FNiagaraMeshRendererMeshPropertiesBase;
 
 /** Reusable extension on cloner layout to group similar options */
 UCLASS(MinimalAPI, Abstract, BlueprintType, Within=CEClonerComponent)
@@ -56,6 +57,15 @@ public:
 		return bExtensionActive;
 	}
 
+	/**
+	 * Called when the meshes are being updated. With opportunity for extension to change the mesh data.
+	 * @return true if the meshes were changed.
+	 */
+	virtual bool OnUpdateClonerMeshes(TArray<FNiagaraMeshRendererMeshPropertiesBase>& InOutMeshes)
+	{
+		return false;
+	}
+
 	/** Called when the meshes are updated */
 	virtual void OnClonerMeshesUpdated() {}
 
diff --git a/Engine/Plugins/VirtualProduction/ClonerEffector/Source/ClonerEffector/Public/Cloner/Extensions/CEClonerMeshRendererExtension.h b/Engine/Plugins/VirtualProduction/ClonerEffector/Source/ClonerEffector/Public/Cloner/Extensions/CEClonerMeshRendererExtension.h
--- a/Engine/Plugins/VirtualProduction/ClonerEffector/Source/ClonerEffector/Public/Cloner/Extensions/CEClonerMeshRendererExtension.h
+++ b/Engine/Plugins/VirtualProduction/ClonerEffector/Source/ClonerEffector/Public/Cloner/Extensions/CEClonerMeshRendererExtension.h
@@ -102,7 +102,7 @@ protected:
 	//~ End UObject
 
 	//~ Begin UCEClonerExtensionBase
-	virtual void OnClonerMeshesUpdated() override;
+	virtual bool OnUpdateClonerMeshes(TArray<FNiagaraMeshRendererMeshPropertiesBase>& InOutMeshes) override;
 	virtual void OnExtensionParametersChanged(UCEClonerComponent* InComponent) override;
 	//~ End UCEClonerExtensionBase
 
