Motion Design: fix new cloner  levels from being kept dirty on first save

#Synced-CL 52400047

diff --git a/Engine/Plugins/VirtualProduction/ClonerEffector/Source/ClonerEffector/Private/Cloner/Layouts/CEClonerLayoutBase.cpp b/Engine/Plugins/VirtualProduction/ClonerEffector/Source/ClonerEffector/Private/Cloner/Layouts/CEClonerLayoutBase.cpp
--- a/Engine/Plugins/VirtualProduction/ClonerEffector/Source/ClonerEffector/Private/Cloner/Layouts/CEClonerLayoutBase.cpp
+++ b/Engine/Plugins/VirtualProduction/ClonerEffector/Source/ClonerEffector/Private/Cloner/Layouts/CEClonerLayoutBase.cpp
@@ -13,6 +13,7 @@
 #include "Misc/Guid.h"
 #include "Misc/PackageName.h"
 #include "NiagaraEmitter.h"
+#include "NiagaraFunctionLibrary.h"
 #include "NiagaraMeshRendererProperties.h"
 #include "NiagaraSystem.h"
 #include "Subsystems/CEClonerSubsystem.h"
@@ -206,7 +207,15 @@ bool UCEClonerLayoutBase::UnloadLayout()
 		return false;
 	}
 
-	UNiagaraDataInterfaceArrayMesh::SetNiagaraArrayMesh(GetClonerComponent(), UE::ClonerEffector::MeshArrayName, {});
+	if (UCEClonerComponent* const ClonerComponent = GetClonerComponent())
+	{
+		const UNiagaraDataInterfaceArrayMesh* const MeshArrayDI = UNiagaraFunctionLibrary::GetDataInterface<UNiagaraDataInterfaceArrayMesh>(ClonerComponent, UE::ClonerEffector::MeshArrayName);
+		// Avoid redundant operation if MeshData is already empty.
+		if (MeshArrayDI && !MeshArrayDI->MeshData.IsEmpty())
+		{
+			UNiagaraDataInterfaceArrayMesh::SetNiagaraArrayMesh(ClonerComponent, UE::ClonerEffector::MeshArrayName, {});
+		}
+	}
 
 	// Clear the meshes as the mesh array DI is going to be used. This should be always empty for new systems, but old systems could have lingering meshes here.
 	MeshRenderer->Meshes.Empty();
@@ -438,7 +447,6 @@ void UCEClonerLayoutBase::CacheMeshRenderer()
 			{
 				if (UNiagaraMeshRendererProperties* EmitterMeshRenderer = Cast<UNiagaraMeshRendererProperties>(EmitterRenderer))
 				{
-					UNiagaraDataInterfaceArrayMesh::SetNiagaraArrayMesh(GetClonerComponent(), UE::ClonerEffector::MeshArrayName, {});
 					MeshRenderer = EmitterMeshRenderer;
 					return;
 				}
