ClonerEffector: fixed crash when updating cloner override materials (specifically MIDs)

#Synced-CL 52400047

diff --git a/Engine/Plugins/VirtualProduction/ClonerEffector/Source/ClonerEffector/Private/Cloner/Extensions/CEClonerMeshRendererExtension.cpp b/Engine/Plugins/VirtualProduction/ClonerEffector/Source/ClonerEffector/Private/Cloner/Extensions/CEClonerMeshRendererExtension.cpp
--- a/Engine/Plugins/VirtualProduction/ClonerEffector/Source/ClonerEffector/Private/Cloner/Extensions/CEClonerMeshRendererExtension.cpp
+++ b/Engine/Plugins/VirtualProduction/ClonerEffector/Source/ClonerEffector/Private/Cloner/Extensions/CEClonerMeshRendererExtension.cpp
@@ -278,8 +278,15 @@ bool UCEClonerMeshRendererExtension::OnUpdateClonerMeshes(TArray<FNiagaraMeshRen
 	}
 
 	// Set material override
-	MeshRenderer->bOverrideMaterials = bUseOverrideMaterial || bVisualizeEffectors;
-	MeshRenderer->OverrideMaterials = GetOverrideMeshesMaterials();
+	TArray<FNiagaraMeshMaterialOverride> OverrideMaterials = GetOverrideMeshesMaterials();
+	const bool bOverrideMaterials = bUseOverrideMaterial || bVisualizeEffectors;
+
+	bMeshDataChanged = bMeshDataChanged 
+		|| MeshRenderer->bOverrideMaterials != bOverrideMaterials
+		|| MeshRenderer->OverrideMaterials != OverrideMaterials;
+
+	MeshRenderer->bOverrideMaterials = bOverrideMaterials;
+	MeshRenderer->OverrideMaterials = MoveTemp(OverrideMaterials);
 
 	return bMeshDataChanged;
 }
